Renamed current VertexDeltaModel to LegacyVertexDeltaModel. Updated VertexDeltaModel with refactored, simplified training code.

diff --git a/Engine/Plugins/Experimental/MLDeformer/LegacyVertexDeltaModel/LegacyVertexDeltaModel.uplugin b/Engine/Plugins/Experimental/MLDeformer/LegacyVertexDeltaModel/LegacyVertexDeltaModel.uplugin
new file mode 100644
--- /dev/null
+++ b/Engine/Plugins/Experimental/MLDeformer/LegacyVertexDeltaModel/LegacyVertexDeltaModel.uplugin
@@ -0,0 +1,51 @@
+{
+	"FileVersion": 3,
+	"Version": 1,
+	"VersionName": "0.2",
+	"FriendlyName": "LegacyVertexDeltaModel",
+	"Description": "Legacy Vertex Delta Model for the ML Deformer Framework",
+	"Category": "Animation",
+	"CreatedBy": "Epic Games, Inc.",
+	"CreatedByURL": "http://epicgames.com",
+	"DocsURL": "",
+	"MarketplaceURL": "",
+	"SupportURL": "",
+	"CanContainContent": true,
+	"IsBetaVersion": false,
+	"IsExperimentalVersion": true,
+	"Installed": false,
+	"Modules": [
+		{
+			"Name": "LegacyVertexDeltaModel",
+			"Type": "Runtime",
+			"LoadingPhase": "Default"
+		},
+		{
+			"Name": "LegacyVertexDeltaModelEditor",
+			"Type": "Editor",
+			"LoadingPhase": "Default"
+		}
+	],
+	"Plugins": [
+		{
+			"Name": "GeometryCache",
+			"Enabled": true
+		},
+		{
+			"Name": "NeuralNetworkInference",
+			"Enabled": true
+		},
+		{
+			"Name": "Optimus",
+			"Enabled": true
+		},
+		{
+			"Name": "PythonFoundationPackages",
+			"Enabled": true
+		},
+		{
+			"Name": "MLDeformerFramework",
+			"Enabled": true
+		}
+	]
+}
\ No newline at end of file
diff --git a/Engine/Plugins/Experimental/MLDeformer/LegacyVertexDeltaModel/Shaders/Private/LegacyVertexDeltaModelDataInterface.ush b/Engine/Plugins/Experimental/MLDeformer/LegacyVertexDeltaModel/Shaders/Private/LegacyVertexDeltaModelDataInterface.ush
new file mode 100644
--- /dev/null
+++ b/Engine/Plugins/Experimental/MLDeformer/LegacyVertexDeltaModel/Shaders/Private/LegacyVertexDeltaModelDataInterface.ush
@@ -0,0 +1,25 @@
+// Copyright Epic Games, Inc. All Rights Reserved.
+
+#include "/Plugin/ComputeFramework/Private/ComputeKernelCommon.ush"
+
+uint DI_LOCAL(NumVertices);
+uint DI_LOCAL(InputStreamStart);
+float3 DI_LOCAL(VertexDeltaScale);
+float3 DI_LOCAL(VertexDeltaMean);
+float DI_LOCAL(Weight);
+
+Buffer<float> DI_LOCAL(PositionDeltaBuffer);
+Buffer<uint> DI_LOCAL(VertexMapBuffer);
+
+DI_IMPL_READ(ReadNumVertices, uint, void)
+{
+	return DI_LOCAL(NumVertices);
+}
+
+DI_IMPL_READ(ReadPositionDelta, float3, uint VertexIndex)
+{
+	uint BufferIndex = DI_LOCAL(VertexMapBuffer)[VertexIndex + DI_LOCAL(InputStreamStart)] * 3;
+	float3 Delta = float3(DI_LOCAL(PositionDeltaBuffer)[BufferIndex], DI_LOCAL(PositionDeltaBuffer)[BufferIndex + 1], DI_LOCAL(PositionDeltaBuffer)[BufferIndex + 2]);
+	float3 FinalDelta = Delta * DI_LOCAL(VertexDeltaScale) + DI_LOCAL(VertexDeltaMean);
+	return FinalDelta * DI_LOCAL(Weight);
+}
diff --git a/Engine/Plugins/Experimental/MLDeformer/LegacyVertexDeltaModel/Source/LegacyVertexDeltaModel/LegacyVertexDeltaModel.Build.cs b/Engine/Plugins/Experimental/MLDeformer/LegacyVertexDeltaModel/Source/LegacyVertexDeltaModel/LegacyVertexDeltaModel.Build.cs
new file mode 100644
--- /dev/null
+++ b/Engine/Plugins/Experimental/MLDeformer/LegacyVertexDeltaModel/Source/LegacyVertexDeltaModel/LegacyVertexDeltaModel.Build.cs
@@ -0,0 +1,58 @@
+// Copyright Epic Games, Inc. All Rights Reserved.
+
+using UnrealBuildTool;
+
+public class LegacyVertexDeltaModel : ModuleRules
+{
+	public LegacyVertexDeltaModel(ReadOnlyTargetRules Target) : base(Target)
+	{
+		PCHUsage = ModuleRules.PCHUsageMode.UseExplicitOrSharedPCHs;
+		
+		PublicIncludePaths.AddRange(
+			new string[] {
+				// ... add public include paths required here ...
+			}
+			);
+				
+		
+		PrivateIncludePaths.AddRange(
+			new string[] {
+				// ... add other private include paths required here ...
+			}
+			);
+			
+		
+		PublicDependencyModuleNames.AddRange(
+			new string[]
+			{
+				"Core"
+				// ... add other public dependencies that you statically link with here ...
+			}
+			);
+			
+		
+		PrivateDependencyModuleNames.AddRange(
+			new string[]
+			{
+				"ComputeFramework",
+				"CoreUObject",
+				"Engine",
+				"GeometryCache",
+				"NeuralNetworkInference",
+				"OptimusCore",
+				"Projects",
+				"RenderCore",
+				"RHI",
+				"MLDeformerFramework"
+			}
+			);
+		
+		
+		DynamicallyLoadedModuleNames.AddRange(
+			new string[]
+			{
+				// ... add any modules that your module loads dynamically here ...
+			}
+			);
+	}
+}
diff --git a/Engine/Plugins/Experimental/MLDeformer/LegacyVertexDeltaModel/Source/LegacyVertexDeltaModel/Private/LegacyVertexDeltaModel.cpp b/Engine/Plugins/Experimental/MLDeformer/LegacyVertexDeltaModel/Source/LegacyVertexDeltaModel/Private/LegacyVertexDeltaModel.cpp
new file mode 100644
--- /dev/null
+++ b/Engine/Plugins/Experimental/MLDeformer/LegacyVertexDeltaModel/Source/LegacyVertexDeltaModel/Private/LegacyVertexDeltaModel.cpp
@@ -0,0 +1,56 @@
+// Copyright Epic Games, Inc. All Rights Reserved.
+#pragma once
+
+#include "LegacyVertexDeltaModel.h"
+#include "LegacyVertexDeltaModelVizSettings.h"
+#include "MLDeformerGeomCacheHelpers.h"
+#include "UObject/Object.h"
+
+#define LOCTEXT_NAMESPACE "ULegacyVertexDeltaModel"
+
+ULegacyVertexDeltaModel::ULegacyVertexDeltaModel(const FObjectInitializer& ObjectInitializer)
+	: Super(ObjectInitializer)
+{
+#if WITH_EDITORONLY_DATA
+	VizSettings = ObjectInitializer.CreateEditorOnlyDefaultSubobject<ULegacyVertexDeltaModelVizSettings>(this, TEXT("VizSettings"));
+#endif
+}
+
+#if WITH_EDITOR
+	void ULegacyVertexDeltaModel::UpdateNumTargetMeshVertices()
+	{
+		NumTargetMeshVerts = UE::MLDeformer::ExtractNumImportedGeomCacheVertices(GeometryCache);
+	}
+#endif // WITH_EDITOR
+
+#if WITH_EDITORONLY_DATA
+	void ULegacyVertexDeltaModel::SampleGroundTruthPositions(float SampleTime, TArray<FVector3f>& OutPositions)
+	{
+		ULegacyVertexDeltaModelVizSettings* VertexVizSettings = Cast<ULegacyVertexDeltaModelVizSettings>(VizSettings);
+		check(VertexVizSettings);
+
+		UGeometryCache* GeomCache = VertexVizSettings->GetTestGroundTruth();
+		if (GeomCache == nullptr)
+		{
+			OutPositions.Reset();
+			return;
+		}
+
+		if (MeshMappings.IsEmpty())
+		{
+			TArray<FString> FailedImportedMeshnames;
+			GenerateGeomCacheMeshMappings(SkeletalMesh, GeomCache, MeshMappings, FailedImportedMeshnames);
+		}
+
+		UE::MLDeformer::SampleGeomCachePositions(
+			0,
+			SampleTime,
+			MeshMappings,
+			SkeletalMesh,
+			GeomCache,
+			AlignmentTransform,
+			OutPositions);
+	}
+#endif
+
+#undef LOCTEXT_NAMESPACE
diff --git a/Engine/Plugins/Experimental/MLDeformer/LegacyVertexDeltaModel/Source/LegacyVertexDeltaModel/Private/LegacyVertexDeltaModelDataInterface.cpp b/Engine/Plugins/Experimental/MLDeformer/LegacyVertexDeltaModel/Source/LegacyVertexDeltaModel/Private/LegacyVertexDeltaModelDataInterface.cpp
new file mode 100644
--- /dev/null
+++ b/Engine/Plugins/Experimental/MLDeformer/LegacyVertexDeltaModel/Source/LegacyVertexDeltaModel/Private/LegacyVertexDeltaModelDataInterface.cpp
@@ -0,0 +1,53 @@
+// Copyright Epic Games, Inc. All Rights Reserved.
+
+#include "LegacyVertexDeltaModelDataInterface.h"
+#include "LegacyVertexDeltaModel.h"
+#include "MLDeformerModelInstance.h"
+#include "MLDeformerAsset.h"
+#include "MLDeformerComponent.h"
+#include "Components/SkeletalMeshComponent.h"
+#include "ComputeFramework/ComputeKernelPermutationSet.h"
+#include "ComputeFramework/ShaderParamTypeDefinition.h"
+#include "ComputeFramework/ShaderParameterMetadataBuilder.h"
+#include "OptimusDataDomain.h"
+#include "RenderGraphBuilder.h"
+#include "RenderGraphUtils.h"
+#include "RenderGraphResources.h"
+#include "SkeletalRenderPublic.h"
+
+BEGIN_SHADER_PARAMETER_STRUCT(FLegacyVertexDeltaModelDataInterfaceParameters,)
+	MLDEFORMER_SHADER_PARAMETERS()
+	SHADER_PARAMETER(FVector3f, VertexDeltaScale)
+	SHADER_PARAMETER(FVector3f, VertexDeltaMean)
+END_SHADER_PARAMETER_STRUCT()
+
+MLDEFORMER_GRAPH_IMPLEMENT_BASICS(
+	ULegacyVertexDeltaModelDataInterface,
+	ULegacyVertexDeltaModelDataProvider,
+	UE::LegacyVertexDeltaModel::FLegacyVertexDeltaModelDataProviderProxy,
+	FLegacyVertexDeltaModelDataInterfaceParameters,
+	TEXT("#include \"/Plugin/LegacyVertexDeltaModel/Private/LegacyVertexDeltaModelDataInterface.ush\"\n"),
+	TEXT("Legacy Vertex Delta Model"))
+
+namespace UE::LegacyVertexDeltaModel
+{
+	FLegacyVertexDeltaModelDataProviderProxy::FLegacyVertexDeltaModelDataProviderProxy(USkeletalMeshComponent* SkeletalMeshComponent, UMLDeformerComponent* DeformerComponent)
+		: UE::MLDeformer::FMLDeformerGraphDataProviderProxy(SkeletalMeshComponent, DeformerComponent)
+	{
+		const UMLDeformerModel* Model = DeformerComponent->GetDeformerAsset()->GetModel();
+		const ULegacyVertexDeltaModel* VertexDeltaModel = Cast<ULegacyVertexDeltaModel>(Model);
+		check(VertexDeltaModel);
+
+		VertexDeltaScale = static_cast<FVector3f>(VertexDeltaModel->GetVertexDeltaScale());
+		VertexDeltaMean = static_cast<FVector3f>(VertexDeltaModel->GetVertexDeltaMean());
+	}
+
+	void FLegacyVertexDeltaModelDataProviderProxy::GatherDispatchData(FDispatchSetup const& InDispatchSetup, FCollectedDispatchData& InOutDispatchData)
+	{
+		MLDEFORMER_GRAPH_DISPATCH_START(FLegacyVertexDeltaModelDataInterfaceParameters, InDispatchSetup, InOutDispatchData)
+		MLDEFORMER_GRAPH_DISPATCH_DEFAULT_PARAMETERS()
+		Parameters->VertexDeltaScale = VertexDeltaScale;
+		Parameters->VertexDeltaMean = VertexDeltaMean;
+		MLDEFORMER_GRAPH_DISPATCH_END()
+	}
+}	// namespace UE::LegacyVertexDeltaModel
diff --git a/Engine/Plugins/Experimental/MLDeformer/LegacyVertexDeltaModel/Source/LegacyVertexDeltaModel/Private/LegacyVertexDeltaModelModule.cpp b/Engine/Plugins/Experimental/MLDeformer/LegacyVertexDeltaModel/Source/LegacyVertexDeltaModel/Private/LegacyVertexDeltaModelModule.cpp
new file mode 100644
--- /dev/null
+++ b/Engine/Plugins/Experimental/MLDeformer/LegacyVertexDeltaModel/Source/LegacyVertexDeltaModel/Private/LegacyVertexDeltaModelModule.cpp
@@ -0,0 +1,25 @@
+// Copyright Epic Games, Inc. All Rights Reserved.
+
+#include "LegacyVertexDeltaModelModule.h"
+#include "Modules/ModuleManager.h"
+#include "Interfaces/IPluginManager.h"
+#include "Misc/Paths.h"
+#include "ShaderCore.h"
+
+LEGACYVERTEXDELTAMODEL_API DEFINE_LOG_CATEGORY(LogLegacyVertexDeltaModel)
+
+#define LOCTEXT_NAMESPACE "LegacyVertexDeltaModelModule"
+
+IMPLEMENT_MODULE(UE::LegacyVertexDeltaModel::FLegacyVertexDeltaModelModule, LegacyVertexDeltaModel)
+
+namespace UE::LegacyVertexDeltaModel
+{
+	void FLegacyVertexDeltaModelModule::StartupModule()
+	{
+		// Register an additional shader path for our shaders used inside the deformer graph system.
+		const FString PluginShaderDir = FPaths::Combine(IPluginManager::Get().FindPlugin(TEXT("LegacyVertexDeltaModel"))->GetBaseDir(), TEXT("Shaders"));
+		AddShaderSourceDirectoryMapping(TEXT("/Plugin/LegacyVertexDeltaModel"), PluginShaderDir);
+	}
+}	// namespace UE::LegacyVertexDeltaModel
+
+#undef LOCTEXT_NAMESPACE
diff --git a/Engine/Plugins/Experimental/MLDeformer/LegacyVertexDeltaModel/Source/LegacyVertexDeltaModel/Private/Tests/LegacyVertexDeltaModelTest.cpp b/Engine/Plugins/Experimental/MLDeformer/LegacyVertexDeltaModel/Source/LegacyVertexDeltaModel/Private/Tests/LegacyVertexDeltaModelTest.cpp
new file mode 100644
--- /dev/null
+++ b/Engine/Plugins/Experimental/MLDeformer/LegacyVertexDeltaModel/Source/LegacyVertexDeltaModel/Private/Tests/LegacyVertexDeltaModelTest.cpp
@@ -0,0 +1,126 @@
+// Copyright Epic Games, Inc. All Rights Reserved.
+
+#include "CoreMinimal.h"
+#include "Misc/AutomationTest.h"
+
+#if WITH_DEV_AUTOMATION_TESTS
+
+#include "LegacyVertexDeltaModelModule.h"
+#include "LegacyVertexDeltaModel.h"
+#include "LegacyVertexDeltaModelVizSettings.h"
+#include "MLDeformerAsset.h"
+#include "MLDeformerComponent.h"
+#include "MLDeformerModelInstance.h"
+#include "MLDeformerGeomCacheHelpers.h"
+#include "UObject/GCObjectScopeGuard.h"
+#include "PreviewScene.h"
+#include "Components/SkeletalMeshComponent.h"
+#include "NeuralNetwork.h"
+#include "GeometryCache.h"
+#include "Engine/SkeletalMesh.h"
+#include "Rendering/SkeletalMeshModel.h"
+#include "Rendering/SkeletalMeshLODModel.h"
+
+namespace UE::LegacyVertexDeltaModelTests
+{
+	IMPLEMENT_SIMPLE_AUTOMATION_TEST(FLegacyVertexDeltaModelMainTest, "MLDeformer.LegacyVertexDeltaModel.MainTest", EAutomationTestFlags::ApplicationContextMask | EAutomationTestFlags::EngineFilter);
+	bool FLegacyVertexDeltaModelMainTest::RunTest(const FString& Parameters)
+	{
+		using namespace UE::MLDeformer;
+		using namespace UE::LegacyVertexDeltaModel;
+
+		UMLDeformerAsset* DeformerAsset = LoadObject<UMLDeformerAsset>(nullptr, TEXT("MLDeformerAsset'/LegacyVertexDeltaModel/Tests/Biceps/biceps_mld.biceps_mld'"));
+		UTEST_NOT_NULL(TEXT("MLDeformerAsset load"), DeformerAsset);
+		FGCObjectScopeGuard Guard(DeformerAsset);
+
+		ULegacyVertexDeltaModel* VertexDeltaModel = Cast<ULegacyVertexDeltaModel>(DeformerAsset->GetModel());
+		UTEST_NOT_NULL(TEXT("LegacyVertexDeltaModel load check"), VertexDeltaModel);
+
+		UNeuralNetwork* NeuralNet = VertexDeltaModel->GetNeuralNetwork();
+		UTEST_NOT_NULL(TEXT("Neuralnet check"), NeuralNet);
+
+		UMLDeformerInputInfo* InputInfo = VertexDeltaModel->GetInputInfo();
+		UTEST_NOT_NULL(TEXT("InputInfo check"), InputInfo);
+		const int64 NumInputInfoInputs = InputInfo->CalcNumNeuralNetInputs();
+		const int64 NumNetworkInputs = NeuralNet->GetInputTensor().Num();
+		UTEST_EQUAL(TEXT("InputInfo input count check"), NumInputInfoInputs, NumNetworkInputs);
+
+		const int32 NumInputBones = InputInfo->GetNumBones();
+		UTEST_EQUAL(TEXT("InputInfo bone count check"), NumInputBones, 1);
+		UTEST_EQUAL(TEXT("InputInfo bone name string check"), InputInfo->GetBoneNameStrings().Num(), 1);
+		UTEST_EQUAL(TEXT("InputInfo bone name string content check"), InputInfo->GetBoneNameString(0), FString("lowerarm_l"));
+		UTEST_EQUAL(TEXT("InputInfo bone names check"), InputInfo->GetBoneNames().Num(), 1);
+		UTEST_EQUAL(TEXT("InputInfo bone names content check"), InputInfo->GetBoneName(0), FName("lowerarm_l"));
+		UTEST_EQUAL(TEXT("InputInfo vertex count check"), InputInfo->GetNumBaseMeshVertices(), InputInfo->GetNumTargetMeshVertices());
+
+		const FVertexMapBuffer& VertexMapBuffer = VertexDeltaModel->GetVertexMapBuffer();
+		UTEST_TRUE(TEXT("VertexMapBuffer check"), VertexMapBuffer.IsInitialized());
+
+#if WITH_EDITORONLY_DATA
+		USkeletalMesh* SkeletalMesh = VertexDeltaModel->GetSkeletalMesh();
+		UTEST_NOT_NULL(TEXT("SkeletalMesh check"), SkeletalMesh);
+
+		UTEST_TRUE(TEXT("Model vs Network compatible check"), VertexDeltaModel->GetInputInfo()->IsCompatible(SkeletalMesh));
+
+		UGeometryCache* GeomCache = VertexDeltaModel->GetGeometryCache();
+		UTEST_NOT_NULL(TEXT("GeomCache check"), GeomCache);
+
+		const FSkeletalMeshModel* ImportedModel = SkeletalMesh->GetImportedModel();
+		UTEST_NOT_NULL(TEXT("ImportedModel check"), ImportedModel);
+
+		const TArray<int32>& VertexMap = VertexDeltaModel->GetVertexMap();
+		UTEST_EQUAL(TEXT("VertexMap size check"), VertexMap.Num(), ImportedModel->LODModels[0].NumVertices);
+
+		const int32 NumBaseMeshVerts = UMLDeformerModel::ExtractNumImportedSkinnedVertices(SkeletalMesh);
+		const int32 NumGeomCacheVerts = ExtractNumImportedGeomCacheVertices(GeomCache);
+		UTEST_EQUAL(TEXT("VertexCount check"), NumBaseMeshVerts, NumGeomCacheVerts);
+		UTEST_EQUAL(TEXT("Model SkelMesh VertexCount check"), VertexDeltaModel->GetNumBaseMeshVerts(), NumBaseMeshVerts);
+		UTEST_EQUAL(TEXT("Model TargetMesh VertexCount check"), VertexDeltaModel->GetNumTargetMeshVerts(), NumGeomCacheVerts);
+
+		const int64 NumNetworkOutputs = NeuralNet->GetOutputTensor().Num();
+		const int64 ExpectedNetworkOutput = NumGeomCacheVerts * 3;	// 3 floats per vertex.
+		UTEST_EQUAL(TEXT("NeuralNet output size check"), NumNetworkOutputs, ExpectedNetworkOutput);
+
+		UTEST_EQUAL(TEXT("Loss function check"), VertexDeltaModel->GetLossFunction(), ELegacyVertexDeltaModelLossFunction::MSE);
+		UTEST_EQUAL(TEXT("Activation function check"), VertexDeltaModel->GetActivationFunction(), ELegacyVertexDeltaModelActivationFunction::LRelu);
+		UTEST_EQUAL(TEXT("Hidden Layers check"), VertexDeltaModel->GetNumHiddenLayers(), 2);
+
+		ULegacyVertexDeltaModelVizSettings* VizSettings = Cast<ULegacyVertexDeltaModelVizSettings>(VertexDeltaModel->GetVizSettings());
+		UTEST_NOT_NULL(TEXT("VizSettings check"), VizSettings);
+
+		// Create a preview scene and spawn an actor.
+		TUniquePtr<FPreviewScene> PreviewScene = MakeUnique<FPreviewScene>(FPreviewScene::ConstructionValues());
+		UWorld* World = PreviewScene->GetWorld();
+		UTEST_NOT_NULL(TEXT("World check"), World);
+		FActorSpawnParameters SpawnParams;
+		SpawnParams.Name = MakeUniqueObjectName(World, AActor::StaticClass(), FName(TEXT("ML Deformer Test Actor")));	
+		AActor* Actor = World->SpawnActor<AActor>(SpawnParams);
+		Actor->SetFlags(RF_Transient);
+		UTEST_NOT_NULL(TEXT("Actor check"), Actor);
+
+		// Add a skeletal mesh component.
+		USkeletalMeshComponent* SkelMeshComponent = NewObject<USkeletalMeshComponent>(Actor);
+		SkelMeshComponent->SetSkeletalMesh(SkeletalMesh);
+		Actor->SetRootComponent(SkelMeshComponent);
+		SkelMeshComponent->RegisterComponent();
+		
+		// Add an ML Deformer component and activate it.
+		UMLDeformerComponent* MLDeformerComponent = NewObject<UMLDeformerComponent>(Actor);
+		MLDeformerComponent->SetDeformerAsset(DeformerAsset);
+		MLDeformerComponent->RegisterComponent();
+		MLDeformerComponent->Activate();
+		UTEST_EQUAL(TEXT("MLDeformerComponent weight check"), MLDeformerComponent->GetWeight(), 1.0f);
+		UTEST_TRUE(TEXT("MLDeformerComponent SkelMeshComponent check"), MLDeformerComponent->GetSkeletalMeshComponent() == SkelMeshComponent);
+
+		// Create and init the model instance.
+		FMLDeformerModelInstance* ModelInstance = MLDeformerComponent->GetModelInstance();
+		UTEST_NOT_NULL(TEXT("CreateModelInstance check"), ModelInstance);
+		UTEST_TRUE(TEXT("ModelInstance GetModel check"), ModelInstance->GetModel() == VertexDeltaModel);
+		UTEST_TRUE(TEXT("ModelInstance compatible check"), ModelInstance->IsCompatible());
+#endif // WITH_EDITORONLY_DATA
+
+		return true;
+	}
+}	// namespace UE::LegacyVertexDeltaModelTests
+
+#endif	// #if WITH_DEV_AUTOMATION_TESTS
diff --git a/Engine/Plugins/Experimental/MLDeformer/LegacyVertexDeltaModel/Source/LegacyVertexDeltaModel/Public/LegacyVertexDeltaModel.h b/Engine/Plugins/Experimental/MLDeformer/LegacyVertexDeltaModel/Source/LegacyVertexDeltaModel/Public/LegacyVertexDeltaModel.h
new file mode 100644
--- /dev/null
+++ b/Engine/Plugins/Experimental/MLDeformer/LegacyVertexDeltaModel/Source/LegacyVertexDeltaModel/Public/LegacyVertexDeltaModel.h
@@ -0,0 +1,139 @@
+// Copyright Epic Games, Inc. All Rights Reserved.
+#pragma once
+
+#include "CoreMinimal.h"
+#include "MLDeformerModel.h"
+#include "MLDeformerVizSettings.h"
+#include "MLDeformerInputInfo.h"
+#include "MLDeformerGeomCacheHelpers.h"
+#include "UObject/Object.h"
+#include "LegacyVertexDeltaModel.generated.h"
+
+class USkeletalMesh;
+class UGeometryCache;
+class UAnimSequence;
+class UNeuralNetwork;
+class UMLDeformerAsset;
+class USkeleton;
+class IPropertyHandle;
+
+/** The activation function to use during the ML Deformer training process. */
+UENUM()
+enum class ELegacyVertexDeltaModelActivationFunction : uint8
+{
+	Relu,
+	LRelu,
+	Tanh
+};
+
+/** The loss function to use during the ML Deformer training process. */
+UENUM()
+enum class ELegacyVertexDeltaModelLossFunction : uint8
+{
+	L1,
+	MSE,
+	Shrinkage
+};
+
+UCLASS()
+class LEGACYVERTEXDELTAMODEL_API ULegacyVertexDeltaModel
+	: public UMLDeformerModel
+{
+	GENERATED_BODY()
+
+public:
+	ULegacyVertexDeltaModel(const FObjectInitializer& ObjectInitializer);
+
+	// UMLDeformerModel overrides.
+	virtual FString GetDisplayName() const override { return "Legacy Vertex Delta Model"; }
+#if WITH_EDITORONLY_DATA
+	virtual bool HasTrainingGroundTruth() const override { return (GeometryCache.Get() != nullptr); }
+	virtual void SampleGroundTruthPositions(float SampleTime, TArray<FVector3f>& OutPositions) override;
+#endif
+#if WITH_EDITOR
+	virtual void UpdateNumTargetMeshVertices();
+#endif
+	// ~END UMLDeformerModel overrides.
+
+	const FVector& GetVertexDeltaMean() const { return VertexDeltaMean; }
+	const FVector& GetVertexDeltaScale() const { return VertexDeltaScale; }
+
+#if WITH_EDITORONLY_DATA
+	const UGeometryCache* GetGeometryCache() const { return GeometryCache; }
+	UGeometryCache* GetGeometryCache() { return GeometryCache; }
+	ELegacyVertexDeltaModelLossFunction GetLossFunction() const { return LossFunction; }
+	ELegacyVertexDeltaModelActivationFunction GetActivationFunction() const { return ActivationFunction; }
+	int32 GetNumHiddenLayers() const { return NumHiddenLayers; }
+	int32 GetNumNeuronsPerLayer() const { return NumNeuronsPerLayer; }
+	int32 GetNumEpochs() const { return Epochs; }
+	int32 GetBatchSize() const { return BatchSize; }
+	float GetLearningRate() const { return LearningRate; }
+	float GetShrinkageSpeed() const { return ShrinkageSpeed; }
+	float GetShrinkageThrehsold() const { return ShrinkageThreshold; }
+	int32 GetMaxCacheSizeGB() const { return MaxCacheSizeGB; }
+	TArray<UE::MLDeformer::FMLDeformerGeomCacheMeshMapping>& GetGeomCacheMeshMappings() { return MeshMappings; }
+	const TArray<UE::MLDeformer::FMLDeformerGeomCacheMeshMapping>& GetGeomCacheMeshMappings() const { return MeshMappings; }
+#endif
+
+public:
+	/** The vertex delta mean. This is passed to the deformer shader to reconstruct the correct output deltas. */
+	UPROPERTY()
+	FVector VertexDeltaMean = FVector::ZeroVector;
+
+	/** The vertex delta scale. This is passed to the deformer shader to rescale the output deltas. */
+	UPROPERTY()
+	FVector VertexDeltaScale = FVector::OneVector;
+
+#if WITH_EDITORONLY_DATA
+	TArray<UE::MLDeformer::FMLDeformerGeomCacheMeshMapping> MeshMappings;
+
+	/** The geometry cache that represents the complex mesh deformations. */
+	UPROPERTY(EditAnywhere, Category = "Target Mesh")
+	TObjectPtr<UGeometryCache> GeometryCache = nullptr;
+
+	/** The number of hidden layers that the neural network model will have.\nHigher numbers will slow down performance but can deal with more complex deformations. */
+	UPROPERTY(EditAnywhere, BlueprintReadWrite, Category = "Training Settings", meta = (ClampMin = "1", ClampMax = "10"))
+	int32 NumHiddenLayers = 2;
+
+	/** The number of units/neurons per hidden layer. Higher numbers will slow down performance but allow for more complex mesh deformations. */
+	UPROPERTY(EditAnywhere, BlueprintReadWrite, Category = "Training Settings", meta = (ClampMin = "1"))
+	int32 NumNeuronsPerLayer = 256;
+
+	/** The number of epochs to process without any decay. */
+	UPROPERTY(EditAnywhere, BlueprintReadWrite, Category = "Training Settings", meta = (ClampMin = "1"))
+	int32 Epochs = 20;
+
+	/** The number of frames per batch when training the model. */
+	UPROPERTY(EditAnywhere, BlueprintReadWrite, AdvancedDisplay, Category = "Training Settings", meta = (ClampMin = "1"))
+	int32 BatchSize = 128;
+
+	/** The learning rate used during the model training. */
+	UPROPERTY(EditAnywhere, BlueprintReadWrite, AdvancedDisplay, Category = "Training Settings", meta = (ClampMin = "0.000001", ClampMax = "1.0"))
+	float LearningRate = 0.00175f;
+
+	/** The activation function to use in the neural network. */
+	UPROPERTY(EditAnywhere, BlueprintReadWrite, AdvancedDisplay, Category = "Training Settings")
+	ELegacyVertexDeltaModelActivationFunction ActivationFunction = ELegacyVertexDeltaModelActivationFunction::LRelu;
+
+	/** The loss function to use during model training. */
+	UPROPERTY(EditAnywhere, BlueprintReadWrite, AdvancedDisplay, Category = "Training Settings")
+	ELegacyVertexDeltaModelLossFunction LossFunction = ELegacyVertexDeltaModelLossFunction::L1;
+
+	/** Shrinkage speed. Only if the shrinkage loss is used. */
+	UPROPERTY(EditAnywhere, BlueprintReadWrite, AdvancedDisplay, Category = "Training Settings", meta = (ClampMin = "1.0", ClampMax = "1000.0"))
+	float ShrinkageSpeed = 10.0f;
+
+	/** Shrinkage threshold. Only if the shrinkage loss is used. */
+	UPROPERTY(EditAnywhere, BlueprintReadWrite, AdvancedDisplay, Category = "Training Settings", meta = (ClampMin = "0.0"))
+	float ShrinkageThreshold = 0.1f;
+
+	/**
+	 * The maximum allowed size of the training cache in memory, in gigabytes.
+	 * So a value of 4 would use a maximum of four gigabyte of system memory.
+	 * The larger the cache size the faster the training.
+	 */
+	UPROPERTY(EditAnywhere, BlueprintReadWrite, AdvancedDisplay, Category = "Training Settings", DisplayName = "Max Cache Size", meta = (ClampMin = "0", ForceUnits = "Gigabytes"))
+	int32 MaxCacheSizeGB = 4;	// 4 Gigabyte
+
+#endif // WITH_EDITORONLY_DATA
+};
diff --git a/Engine/Plugins/Experimental/MLDeformer/LegacyVertexDeltaModel/Source/LegacyVertexDeltaModel/Public/LegacyVertexDeltaModelDataInterface.h b/Engine/Plugins/Experimental/MLDeformer/LegacyVertexDeltaModel/Source/LegacyVertexDeltaModel/Public/LegacyVertexDeltaModelDataInterface.h
new file mode 100644
--- /dev/null
+++ b/Engine/Plugins/Experimental/MLDeformer/LegacyVertexDeltaModel/Source/LegacyVertexDeltaModel/Public/LegacyVertexDeltaModelDataInterface.h
@@ -0,0 +1,69 @@
+// Copyright Epic Games, Inc. All Rights Reserved.
+
+#pragma once
+
+#include "CoreMinimal.h"
+#include "MLDeformerGraphDataInterface.h"
+#include "OptimusComputeDataInterface.h"
+#include "ComputeFramework/ComputeDataProvider.h"
+#include "LegacyVertexDeltaModelDataInterface.generated.h"
+
+class FRDGBuffer;
+class FRDGBufferSRV;
+class FRHIShaderResourceView;
+class FSkeletalMeshObject;
+class UMLDeformerComponent;
+class UNeuralNetwork;
+class USkeletalMeshComponent;
+class UMLDeformerModel;
+
+/** Compute Framework Data Interface for MLDefomer data. */
+UCLASS(Category = ComputeFramework)
+class LEGACYVERTEXDELTAMODEL_API ULegacyVertexDeltaModelDataInterface
+	: public UMLDeformerGraphDataInterface
+{
+	GENERATED_BODY()
+
+public:
+	// UOptimusComputeDataInterface overrides.
+	virtual FString GetDisplayName() const override;
+	// ~END UOptimusComputeDataInterface overrides.
+
+	// UComputeDataInterface overrides.
+	virtual void GetShaderParameters(TCHAR const* UID, FShaderParametersMetadataBuilder& OutBuilder) const override;
+	virtual void GetHLSL(FString& OutHLSL) const override;
+	virtual UComputeDataProvider* CreateDataProvider(TArrayView<TObjectPtr<UObject>> InSourceObjects, uint64 InInputMask, uint64 InOutputMask) const override;
+	// ~END UComputeDataInterface overrides.
+};
+
+/** Compute Framework Data Provider for MLDeformer data. */
+UCLASS(BlueprintType, EditInlineNew, Category = ComputeFramework)
+class LEGACYVERTEXDELTAMODEL_API ULegacyVertexDeltaModelDataProvider
+	: public UMLDeformerGraphDataProvider
+{
+	GENERATED_BODY()
+
+public:
+	// UComputeDataProvider overrides.
+	virtual FComputeDataProviderRenderProxy* GetRenderProxy() override;
+	// ~END UComputeDataProvider overrides.
+};
+
+namespace UE::LegacyVertexDeltaModel
+{
+	/** Compute Framework Data Provider Proxy for MLDeformer data. */
+	class LEGACYVERTEXDELTAMODEL_API FLegacyVertexDeltaModelDataProviderProxy
+		: public UE::MLDeformer::FMLDeformerGraphDataProviderProxy
+	{
+	public:
+		FLegacyVertexDeltaModelDataProviderProxy(USkeletalMeshComponent* SkeletalMeshComponent, UMLDeformerComponent* DeformerComponent);
+
+		// FComputeDataProviderRenderProxy overrides.
+		virtual void GatherDispatchData(FDispatchSetup const& InDispatchSetup, FCollectedDispatchData& InOutDispatchData) override;
+		// ~END FComputeDataProviderRenderProxy overrides.
+
+	private:
+		FVector3f VertexDeltaScale = FVector3f(1.0f, 1.0f, 1.0f);
+		FVector3f VertexDeltaMean = FVector3f(0.0f, 0.0f, 0.0f);
+	};
+}	// namespace UE::LegacyVertexDeltaModel
diff --git a/Engine/Plugins/Experimental/MLDeformer/LegacyVertexDeltaModel/Source/LegacyVertexDeltaModel/Public/LegacyVertexDeltaModelModule.h b/Engine/Plugins/Experimental/MLDeformer/LegacyVertexDeltaModel/Source/LegacyVertexDeltaModel/Public/LegacyVertexDeltaModelModule.h
new file mode 100644
--- /dev/null
+++ b/Engine/Plugins/Experimental/MLDeformer/LegacyVertexDeltaModel/Source/LegacyVertexDeltaModel/Public/LegacyVertexDeltaModelModule.h
@@ -0,0 +1,20 @@
+// Copyright Epic Games, Inc. All Rights Reserved.
+
+#pragma once
+
+#include "CoreMinimal.h"
+#include "Modules/ModuleManager.h"
+
+LEGACYVERTEXDELTAMODEL_API DECLARE_LOG_CATEGORY_EXTERN(LogLegacyVertexDeltaModel, Log, All);
+
+namespace UE::LegacyVertexDeltaModel
+{
+	class LEGACYVERTEXDELTAMODEL_API FLegacyVertexDeltaModelModule
+		: public IModuleInterface
+	{
+	public:
+		// IModuleInterface overrides.
+		virtual void StartupModule() override;
+		// ~END IModuleInterface overrides.
+	};
+}	// namespace VertexDeltaModel
diff --git a/Engine/Plugins/Experimental/MLDeformer/LegacyVertexDeltaModel/Source/LegacyVertexDeltaModel/Public/LegacyVertexDeltaModelVizSettings.h b/Engine/Plugins/Experimental/MLDeformer/LegacyVertexDeltaModel/Source/LegacyVertexDeltaModel/Public/LegacyVertexDeltaModelVizSettings.h
new file mode 100644
--- /dev/null
+++ b/Engine/Plugins/Experimental/MLDeformer/LegacyVertexDeltaModel/Source/LegacyVertexDeltaModel/Public/LegacyVertexDeltaModelVizSettings.h
@@ -0,0 +1,30 @@
+// Copyright Epic Games, Inc. All Rights Reserved.
+
+#pragma once
+
+#include "CoreMinimal.h"
+#include "UObject/Object.h"
+#include "LegacyVertexDeltaModelVizSettings.generated.h"
+
+class UGeometryCache;
+
+/**
+ * The vizualization settings specific to the the vertex delta model.
+ */
+UCLASS()
+class LEGACYVERTEXDELTAMODEL_API ULegacyVertexDeltaModelVizSettings
+	: public UMLDeformerVizSettings
+{
+	GENERATED_BODY()
+
+public:
+#if WITH_EDITORONLY_DATA
+	virtual bool HasTestGroundTruth() const override { return (GroundTruth.Get() != nullptr); }
+	UGeometryCache* GetTestGroundTruth() const { return GroundTruth; }
+
+public:
+	/** The geometry cache that represents the ground truth of the test anim sequence. */
+	UPROPERTY(EditAnywhere, Category = "Test Assets")
+	TObjectPtr<UGeometryCache> GroundTruth = nullptr;
+#endif // WITH_EDITORONLY_DATA
+};
diff --git a/Engine/Plugins/Experimental/MLDeformer/LegacyVertexDeltaModel/Source/LegacyVertexDeltaModelEditor/LegacyVertexDeltaModelEditor.Build.cs b/Engine/Plugins/Experimental/MLDeformer/LegacyVertexDeltaModel/Source/LegacyVertexDeltaModelEditor/LegacyVertexDeltaModelEditor.Build.cs
new file mode 100644
--- /dev/null
+++ b/Engine/Plugins/Experimental/MLDeformer/LegacyVertexDeltaModel/Source/LegacyVertexDeltaModelEditor/LegacyVertexDeltaModelEditor.Build.cs
@@ -0,0 +1,47 @@
+// Copyright Epic Games, Inc. All Rights Reserved.
+
+using System.IO;
+
+namespace UnrealBuildTool.Rules
+{
+	public class LegacyVertexDeltaModelEditor : ModuleRules
+	{
+		public LegacyVertexDeltaModelEditor(ReadOnlyTargetRules Target) : base(Target)
+		{
+			PublicDependencyModuleNames.AddRange(
+				new string[]
+				{
+					"Core"
+				}
+			);
+
+			PrivateDependencyModuleNames.AddRange(
+				new string[]
+				{
+					"EditorFramework",
+					"UnrealEd",
+					"CoreUObject",
+					"Engine",
+					"Slate",
+					"SlateCore",
+					"EditorWidgets",
+					"EditorStyle",
+					"GeometryCache",
+					"MLDeformerFramework",
+					"MLDeformerFrameworkEditor",
+					"LegacyVertexDeltaModel",
+					"PropertyEditor",
+					"NeuralNetworkInference",
+					"ToolWidgets",
+					"ComputeFramework"
+				}
+			);
+
+			PrivateIncludePathModuleNames.AddRange(
+				new string[]
+				{
+				}
+			);
+		}
+	}
+}
diff --git a/Engine/Plugins/Experimental/MLDeformer/LegacyVertexDeltaModel/Source/LegacyVertexDeltaModelEditor/Private/LegacyVertexDeltaEditorModel.cpp b/Engine/Plugins/Experimental/MLDeformer/LegacyVertexDeltaModel/Source/LegacyVertexDeltaModelEditor/Private/LegacyVertexDeltaEditorModel.cpp
new file mode 100644
--- /dev/null
+++ b/Engine/Plugins/Experimental/MLDeformer/LegacyVertexDeltaModel/Source/LegacyVertexDeltaModelEditor/Private/LegacyVertexDeltaEditorModel.cpp
@@ -0,0 +1,305 @@
+// Copyright Epic Games, Inc. All Rights Reserved.
+
+#include "LegacyVertexDeltaEditorModel.h"
+#include "LegacyVertexDeltaModel.h"
+#include "LegacyVertexDeltaModelVizSettings.h"
+#include "LegacyVertexDeltaTrainingModel.h"
+#include "LegacyVertexDeltaEditorModelActor.h"
+#include "MLDeformerEditorToolkit.h"
+#include "MLDeformerEditorStyle.h"
+#include "MLDeformerInputInfo.h"
+#include "MLDeformerGeomCacheSampler.h"
+#include "MLDeformerGeomCacheHelpers.h"
+#include "GeometryCache.h"
+#include "GeometryCacheComponent.h"
+#include "Animation/DebugSkelMeshComponent.h"
+#include "Rendering/SkeletalMeshModel.h"
+#include "Rendering/SkeletalMeshLODModel.h"
+#include "Engine/SkeletalMesh.h"
+
+#define LOCTEXT_NAMESPACE "LegacyVertexDeltaEditorModel"
+
+namespace UE::LegacyVertexDeltaModel
+{
+	using namespace UE::MLDeformer;
+
+	FLegacyVertexDeltaEditorModel::~FLegacyVertexDeltaEditorModel()
+	{
+		Model->OnPostEditChangeProperty().Unbind();
+	}
+
+	FMLDeformerEditorModel* FLegacyVertexDeltaEditorModel::MakeInstance()
+	{
+		return new FLegacyVertexDeltaEditorModel();
+	}
+
+	void FLegacyVertexDeltaEditorModel::AddReferencedObjects(FReferenceCollector& Collector)
+	{
+		Collector.AddReferencedObject(EditorInputInfo);
+	}
+
+	FMLDeformerSampler* FLegacyVertexDeltaEditorModel::CreateSampler() const
+	{ 
+		FMLDeformerGeomCacheSampler* NewSampler = new FMLDeformerGeomCacheSampler();
+		NewSampler->OnGetGeometryCache().BindLambda([this] { return GetVertexDeltaModel()->GetGeometryCache(); });
+		return NewSampler;
+	}
+
+	void FLegacyVertexDeltaEditorModel::Init(const InitSettings& InitSettings)
+	{
+		FMLDeformerEditorModel::Init(InitSettings);
+		Model->OnPostEditChangeProperty().BindRaw(this, &FLegacyVertexDeltaEditorModel::OnPostEditChangeProperty);
+	}
+
+	void FLegacyVertexDeltaEditorModel::OnPostEditChangeProperty(FPropertyChangedEvent& PropertyChangedEvent)
+	{
+		const FProperty* Property = PropertyChangedEvent.Property;
+		if (Property == nullptr)
+		{
+			return;
+		}
+
+		HandleDefaultPropertyChanges(PropertyChangedEvent);
+
+		// When we change one of these properties below, restart animations etc.
+		if (Property->GetFName() == GET_MEMBER_NAME_CHECKED(ULegacyVertexDeltaModel, GeometryCache) ||
+			Property->GetFName() == GET_MEMBER_NAME_CHECKED(ULegacyVertexDeltaModelVizSettings, GroundTruth))
+		{
+			TriggerInputAssetChanged(true);
+		}
+		else
+		if (Property->GetFName() == GET_MEMBER_NAME_CHECKED(ULegacyVertexDeltaModel, DeltaCutoffLength) ||
+			Property->GetFName() == GET_MEMBER_NAME_CHECKED(UMLDeformerModel, AlignmentTransform))
+		{
+			if (PropertyChangedEvent.ChangeType == EPropertyChangeType::ValueSet)
+			{
+				bIsDataNormalized = false;
+				SampleDeltas();
+			}
+		}
+	}
+
+	void FLegacyVertexDeltaEditorModel::OnInputAssetsChanged()
+	{
+		// Update the skeletal mesh components of the training, test base, and ml deformed actor.
+		FMLDeformerEditorModel::OnInputAssetsChanged();
+
+		ULegacyVertexDeltaModelVizSettings* VizSettings = GetVertexDeltaModelVizSettings();
+		check(VizSettings);
+		const float TestAnimSpeed = VizSettings->GetAnimPlaySpeed();
+		UAnimSequence* TestAnimSequence = VizSettings->GetTestAnimSequence();
+
+		// Update the training geometry cache.
+		UGeometryCacheComponent* GeometryCacheComponent = FindVertexDeltaModelEditorActor(ActorID_Train_GroundTruth)->GetGeometryCacheComponent();
+		check(GeometryCacheComponent);
+		GeometryCacheComponent->SetGeometryCache(GetVertexDeltaModel()->GetGeometryCache());
+		GeometryCacheComponent->ResetAnimationTime();
+		GeometryCacheComponent->SetLooping(false);
+		GeometryCacheComponent->SetManualTick(true);
+		GeometryCacheComponent->SetPlaybackSpeed(TestAnimSpeed);
+		GeometryCacheComponent->Play();
+
+		// Update the test geometry cache (ground truth) component.
+		GeometryCacheComponent = FindVertexDeltaModelEditorActor(ActorID_Test_GroundTruth)->GetGeometryCacheComponent();
+		GeometryCacheComponent->SetGeometryCache(VizSettings->GetTestGroundTruth());
+		GeometryCacheComponent->ResetAnimationTime();
+		GeometryCacheComponent->SetLooping(true);
+		GeometryCacheComponent->SetManualTick(true);
+		GeometryCacheComponent->SetPlaybackSpeed(TestAnimSpeed);
+		GeometryCacheComponent->Play();
+
+		ULegacyVertexDeltaModel* VertexDeltaModel = GetVertexDeltaModel();
+		VertexDeltaModel->MeshMappings.Reset();
+	}
+
+	void FLegacyVertexDeltaEditorModel::CreateTrainingGroundTruthActor(UWorld* World)
+	{
+		UGeometryCache* GeomCache = GetVertexDeltaModel()->GetGeometryCache();
+		const FLinearColor LabelColor = FMLDeformerEditorStyle::Get().GetColor("MLDeformer.TargetMesh.LabelColor");
+		const FLinearColor WireframeColor = FMLDeformerEditorStyle::Get().GetColor("MLDeformer.TargetMesh.WireframeColor");
+		CreateGeomCacheActor(
+			World, 
+			ActorID_Train_GroundTruth, 
+			"Train GroundTruth", 
+			GeomCache, 
+			LabelColor, 
+			WireframeColor, 
+			LOCTEXT("TrainGroundTruthActorLabelText", "Target Mesh"),
+			true);
+	}
+
+	void FLegacyVertexDeltaEditorModel::CreateTestGroundTruthActor(UWorld* World)
+	{
+		UGeometryCache* GeomCache = GetVertexDeltaModelVizSettings()->GetTestGroundTruth();
+		const FLinearColor LabelColor = FMLDeformerEditorStyle::Get().GetColor("MLDeformer.GroundTruth.LabelColor");
+		const FLinearColor WireframeColor = FMLDeformerEditorStyle::Get().GetColor("MLDeformer.GroundTruth.WireframeColor");
+		CreateGeomCacheActor(
+			World, 
+			ActorID_Test_GroundTruth, 
+			"Test GroundTruth", 
+			GeomCache, 
+			LabelColor, 
+			WireframeColor,
+			LOCTEXT("TestGroundTruthActorLabelText", "Ground Truth"),
+			false);
+	}
+
+	void FLegacyVertexDeltaEditorModel::CreateGeomCacheActor(UWorld* World, int32 ActorID, const FName& Name, UGeometryCache* GeomCache, FLinearColor LabelColor, FLinearColor WireframeColor, const FText& LabelText, bool bIsTrainingActor)
+	{
+		FActorSpawnParameters SpawnParams;
+		SpawnParams.Name = MakeUniqueObjectName(World, AActor::StaticClass(), Name);
+		AActor* Actor = World->SpawnActor<AActor>(SpawnParams);
+		Actor->SetFlags(RF_Transient);
+
+		// Create the Geometry Cache Component.
+		UGeometryCacheComponent* GeomCacheComponent = NewObject<UGeometryCacheComponent>(Actor);
+		GeomCacheComponent->SetGeometryCache(GeomCache);
+		GeomCacheComponent->RegisterComponent();
+		GeomCacheComponent->SetOverrideWireframeColor(true);
+		GeomCacheComponent->SetWireframeOverrideColor(WireframeColor);
+		GeomCacheComponent->MarkRenderStateDirty();
+		GeomCacheComponent->SetVisibility(false);
+		Actor->SetRootComponent(GeomCacheComponent);
+
+		// Create the editor actor.
+		FMLDeformerEditorActor::FConstructSettings Settings;
+		Settings.Actor = Actor;
+		Settings.TypeID = ActorID;
+		Settings.LabelColor = LabelColor;
+		Settings.LabelText = LabelText;
+		Settings.bIsTrainingActor = bIsTrainingActor;
+		FLegacyVertexDeltaEditorModelActor* EditorActor = static_cast<FLegacyVertexDeltaEditorModelActor*>(CreateEditorActor(Settings));
+		EditorActor->SetGeometryCacheComponent(GeomCacheComponent);
+		EditorActors.Add(EditorActor);
+	} 
+
+	FMLDeformerEditorActor* FLegacyVertexDeltaEditorModel::CreateEditorActor(const FMLDeformerEditorActor::FConstructSettings& Settings) const
+	{
+		return new FLegacyVertexDeltaEditorModelActor(Settings);
+	}
+
+	FLegacyVertexDeltaEditorModelActor* FLegacyVertexDeltaEditorModel::FindVertexDeltaModelEditorActor(int32 TypeID) const
+	{
+		return static_cast<FLegacyVertexDeltaEditorModelActor*>(FindEditorActor(TypeID));
+	}
+
+	ULegacyVertexDeltaModelVizSettings* FLegacyVertexDeltaEditorModel::GetVertexDeltaModelVizSettings() const
+	{
+		return Cast<ULegacyVertexDeltaModelVizSettings>(GetVertexDeltaModel()->GetVizSettings());
+	}
+
+	double FLegacyVertexDeltaEditorModel::GetTimeAtFrame(int32 FrameNumber) const
+	{
+		const FLegacyVertexDeltaEditorModelActor* EditorActor = static_cast<FLegacyVertexDeltaEditorModelActor*>(GetTimelineEditorActor());
+		if (EditorActor && EditorActor->GetGeometryCacheComponent())
+		{
+			return EditorActor->GetGeometryCacheComponent()->GetTimeAtFrame(FrameNumber);
+		}
+		return 0.0;
+	}
+
+	int32 FLegacyVertexDeltaEditorModel::GetFrameAtTime(double TimeInSeconds) const
+	{
+		const FLegacyVertexDeltaEditorModelActor* EditorActor = static_cast<FLegacyVertexDeltaEditorModelActor*>(GetTimelineEditorActor());
+		if (EditorActor && EditorActor->GetGeometryCacheComponent())
+		{
+			return EditorActor->GetGeometryCacheComponent()->GetFrameAtTime(TimeInSeconds);
+		}
+		return 0;
+	}
+
+	int32 FLegacyVertexDeltaEditorModel::GetNumFrames() const
+	{
+		const UGeometryCache* GeometryCache = GetVertexDeltaModel()->GetGeometryCache();
+		if (GeometryCache == nullptr)
+		{
+			return 0;
+		}
+		const int32 StartFrame = GeometryCache->GetStartFrame();
+		const int32 EndFrame = GeometryCache->GetEndFrame();
+		return (EndFrame - StartFrame) + 1;
+	}
+
+	void FLegacyVertexDeltaEditorModel::UpdateIsReadyForTrainingState()
+	{
+		bIsReadyForTraining = false;
+
+		// Do some basic checks first, like if there is a skeletal mesh, ground truth, anim sequence, and if there are frames.
+		if (!IsEditorReadyForTrainingBasicChecks())
+		{
+			return;
+		}
+
+		// Now make sure the assets are compatible.
+		ULegacyVertexDeltaModel* VertexDeltaModel = GetVertexDeltaModel();
+		UGeometryCache* GeomCache = VertexDeltaModel->GetGeometryCache();
+		USkeletalMesh* SkeletalMesh = VertexDeltaModel->GetSkeletalMesh();
+		if (!GetGeomCacheVertexErrorText(SkeletalMesh, GeomCache, FText(), FText()).IsEmpty() ||
+			!GetGeomCacheErrorText(SkeletalMesh, GeomCache).IsEmpty())
+		{
+			return;
+		}
+
+		// Make sure every skeletal imported mesh has some geometry track.
+		const int32 NumGeomCacheTracks = GeomCache ? GeomCache->Tracks.Num() : 0;
+		int32 NumSkelMeshes = 0;
+		if (SkeletalMesh)
+		{
+			FSkeletalMeshModel* ImportedModel = SkeletalMesh->GetImportedModel();
+			if (ImportedModel)
+			{
+				NumSkelMeshes = ImportedModel->LODModels[0].ImportedMeshInfos.Num();		
+			}
+		}
+
+		// Allow the special case where there is just one mesh and track.
+		if (NumGeomCacheTracks != 1 || NumSkelMeshes != 1)
+		{
+			if (!GetGeomCacheSampler()->GetFailedImportedMeshNames().IsEmpty())
+			{
+				return;
+			}
+		}
+
+		bIsReadyForTraining = true;
+	}
+
+	ETrainingResult FLegacyVertexDeltaEditorModel::Train()
+	{
+		return TrainModel<ULegacyVertexDeltaTrainingModel>(this);
+	}
+
+	void FLegacyVertexDeltaEditorModel::OnPreTraining()
+	{
+		// Make a backup of the normalization values, as they get overwritten when training.
+		// However, when we abort, we want to restore to the original values again.
+		// See OnTrainingAborted for when we restore the backup again.
+		VertexDeltaMeanBackup = GetVertexDeltaModel()->GetVertexDeltaMean();
+		VertexDeltaScaleBackup = GetVertexDeltaModel()->GetVertexDeltaScale();
+	}
+
+	void FLegacyVertexDeltaEditorModel::OnTrainingAborted()
+	{
+		// Restore the vertex delta mean and scale, as we aborted, and they could have changed when training
+		// on a smaller subset of frames/samples. If we don't do this, the mesh will deform incorrectly.
+		GetVertexDeltaModel()->VertexDeltaMean = VertexDeltaMeanBackup;
+		GetVertexDeltaModel()->VertexDeltaScale = VertexDeltaScaleBackup;
+	}
+
+	FString FLegacyVertexDeltaEditorModel::GetTrainedNetworkOnnxFile() const
+	{
+		return FString(FPaths::ProjectIntermediateDir() + TEXT("LegacyVertexDeltaModel/latest_net_G.onnx"));
+	}
+
+	FString FLegacyVertexDeltaEditorModel::GetDefaultDeformerGraphAssetPath() const
+	{
+		return FString(TEXT("/LegacyVertexDeltaModel/Deformers/DG_LegacyVertexDeltaModel.DG_LegacyVertexDeltaModel"));
+	}
+
+	FString FLegacyVertexDeltaEditorModel::GetHeatMapDeformerGraphPath() const
+	{
+		return FString(TEXT("/LegacyVertexDeltaModel/Deformers/DG_LegacyVertexDeltaModel_HeatMap.DG_LegacyVertexDeltaModel_HeatMap"));
+	}
+}	// namespace UE::VertexDeltaModel
+
+#undef LOCTEXT_NAMESPACE
diff --git a/Engine/Plugins/Experimental/MLDeformer/LegacyVertexDeltaModel/Source/LegacyVertexDeltaModelEditor/Private/LegacyVertexDeltaEditorModelActor.cpp b/Engine/Plugins/Experimental/MLDeformer/LegacyVertexDeltaModel/Source/LegacyVertexDeltaModelEditor/Private/LegacyVertexDeltaEditorModelActor.cpp
new file mode 100644
--- /dev/null
+++ b/Engine/Plugins/Experimental/MLDeformer/LegacyVertexDeltaModel/Source/LegacyVertexDeltaModelEditor/Private/LegacyVertexDeltaEditorModelActor.cpp
@@ -0,0 +1,116 @@
+// Copyright Epic Games, Inc. All Rights Reserved.
+
+#include "LegacyVertexDeltaEditorModelActor.h"
+#include "GeometryCacheComponent.h"
+#include "Animation/DebugSkelMeshComponent.h"
+
+namespace UE::LegacyVertexDeltaModel
+{
+	using namespace UE::MLDeformer;
+
+	FLegacyVertexDeltaEditorModelActor::FLegacyVertexDeltaEditorModelActor(const FConstructSettings& Settings)
+		: FMLDeformerEditorActor(Settings)
+	{
+	}
+
+	void FLegacyVertexDeltaEditorModelActor::SetVisibility(bool bIsVisible)
+	{
+		FMLDeformerEditorActor::SetVisibility(bIsVisible);
+
+		if (GeomCacheComponent && bIsVisible != GeomCacheComponent->IsVisible())
+		{
+			GeomCacheComponent->SetVisibility(bIsVisible, true);
+		}
+	}
+
+	bool FLegacyVertexDeltaEditorModelActor::IsVisible() const
+	{
+		if (SkeletalMeshComponent)
+		{
+			return SkeletalMeshComponent->IsVisible();
+		}
+
+		if (GeomCacheComponent)
+		{
+			return GeomCacheComponent->IsVisible();
+		}
+
+		return true;
+	}
+
+	void FLegacyVertexDeltaEditorModelActor::SetPlayPosition(float TimeInSeconds, bool bAutoPause)
+	{
+		if (SkeletalMeshComponent)
+		{
+			SkeletalMeshComponent->SetPosition(TimeInSeconds);
+			if (bAutoPause)
+			{
+				SkeletalMeshComponent->bPauseAnims = true;
+			}
+		}
+
+		if (GeomCacheComponent)
+		{
+			GeomCacheComponent->SetManualTick(true);
+			GeomCacheComponent->TickAtThisTime(TimeInSeconds, false, false, false);
+		}
+	}
+
+	float FLegacyVertexDeltaEditorModelActor::GetPlayPosition() const
+	{
+		if (SkeletalMeshComponent)
+		{
+			return SkeletalMeshComponent->GetPosition();
+		}
+
+		if (GeomCacheComponent)
+		{
+			return GeomCacheComponent->GetAnimationTime();
+		}
+
+		return 0.0f;
+	}
+
+	void FLegacyVertexDeltaEditorModelActor::SetPlaySpeed(float PlaySpeed)
+	{
+		if (SkeletalMeshComponent)
+		{
+			SkeletalMeshComponent->SetPlayRate(PlaySpeed);
+		}
+
+		if (GeomCacheComponent)
+		{
+			GeomCacheComponent->SetPlaybackSpeed(PlaySpeed);
+		}
+	}
+
+	void FLegacyVertexDeltaEditorModelActor::Pause(bool bPaused)
+	{
+		if (SkeletalMeshComponent)
+		{
+			SkeletalMeshComponent->bPauseAnims = bPaused;
+		}
+
+		if (GeomCacheComponent)
+		{
+			GeomCacheComponent->SetManualTick(bPaused);
+		}
+	}
+
+	FBox FLegacyVertexDeltaEditorModelActor::GetBoundingBox() const
+	{
+		if (SkeletalMeshComponent)
+		{
+			return SkeletalMeshComponent->Bounds.GetBox();
+		}
+
+		if (GeomCacheComponent)
+		{
+			return GeomCacheComponent->Bounds.GetBox();
+		}
+
+		FBox Box;
+		Box.Init();
+		return Box;
+	}
+}	// namespace UE::LegacyVertexDeltaModel
diff --git a/Engine/Plugins/Experimental/MLDeformer/LegacyVertexDeltaModel/Source/LegacyVertexDeltaModelEditor/Private/LegacyVertexDeltaModelDetails.cpp b/Engine/Plugins/Experimental/MLDeformer/LegacyVertexDeltaModel/Source/LegacyVertexDeltaModelEditor/Private/LegacyVertexDeltaModelDetails.cpp
new file mode 100644
--- /dev/null
+++ b/Engine/Plugins/Experimental/MLDeformer/LegacyVertexDeltaModel/Source/LegacyVertexDeltaModelEditor/Private/LegacyVertexDeltaModelDetails.cpp
@@ -0,0 +1,187 @@
+// Copyright Epic Games, Inc. All Rights Reserved.
+
+#include "LegacyVertexDeltaModelDetails.h"
+#include "LegacyVertexDeltaEditorModel.h"
+#include "LegacyVertexDeltaModel.h"
+#include "MLDeformerModule.h"
+#include "MLDeformerAsset.h"
+#include "MLDeformerModel.h"
+#include "MLDeformerEditorModel.h"
+#include "MLDeformerEditorModule.h"
+#include "MLDeformerGeomCacheHelpers.h"
+#include "NeuralNetwork.h"
+#include "GeometryCache.h"
+#include "GeometryCacheTrack.h"
+#include "Animation/AnimSequence.h"
+#include "Engine/SkeletalMesh.h"
+#include "PropertyCustomizationHelpers.h"
+#include "AssetRegistry/AssetData.h"
+#include "Modules/ModuleManager.h"
+#include "DetailLayoutBuilder.h"
+#include "DetailCategoryBuilder.h"
+#include "DetailWidgetRow.h"
+#include "IDetailsView.h"
+#include "IDetailGroup.h"
+#include "Widgets/DeclarativeSyntaxSupport.h"
+#include "Widgets/Layout/SBox.h"
+#include "Widgets/Input/SButton.h"
+#include "DetailLayoutBuilder.h"
+#include "SWarningOrErrorBox.h"
+
+#define LOCTEXT_NAMESPACE "LegacyVertexDeltaModelDetails"
+
+namespace UE::LegacyVertexDeltaModel
+{
+	using namespace UE::MLDeformer;
+
+	TSharedRef<IDetailCustomization> FLegacyVertexDeltaModelDetails::MakeInstance()
+	{
+		return MakeShareable(new FLegacyVertexDeltaModelDetails());
+	}
+
+	bool FLegacyVertexDeltaModelDetails::UpdateMemberPointers(const TArray<TWeakObjectPtr<UObject>>& Objects)
+	{
+		if (!FMLDeformerModelDetails::UpdateMemberPointers(Objects))
+		{
+			return false;
+		}
+
+		VertexModel = Cast<ULegacyVertexDeltaModel>(Model);
+		check(VertexModel);
+		VertexEditorModel = static_cast<FLegacyVertexDeltaEditorModel*>(EditorModel);
+
+		return (VertexModel != nullptr && VertexEditorModel != nullptr);
+	}
+
+	void FLegacyVertexDeltaModelDetails::CustomizeDetails(IDetailLayoutBuilder& DetailBuilder)
+	{
+		// Create all the detail categories and add the properties of the base class.
+		FMLDeformerModelDetails::CustomizeDetails(DetailBuilder);
+
+		// Training settings.
+		SettingsCategoryBuilder->AddProperty(GET_MEMBER_NAME_CHECKED(ULegacyVertexDeltaModel, NumHiddenLayers), UMLDeformerModel::StaticClass());
+		SettingsCategoryBuilder->AddProperty(GET_MEMBER_NAME_CHECKED(ULegacyVertexDeltaModel, NumNeuronsPerLayer), UMLDeformerModel::StaticClass());
+		SettingsCategoryBuilder->AddProperty(GET_MEMBER_NAME_CHECKED(ULegacyVertexDeltaModel, Epochs), UMLDeformerModel::StaticClass());
+
+		// Check whether shrinkage settings should be visible or not.
+		auto IsShrinkageVisible = [this]()
+		{
+			return (VertexModel->GetLossFunction() == ELegacyVertexDeltaModelLossFunction::Shrinkage) ? EVisibility::Visible : EVisibility::Collapsed;
+		};
+	 
+		// Advanced settings.
+		SettingsCategoryBuilder->AddProperty(GET_MEMBER_NAME_CHECKED(ULegacyVertexDeltaModel, BatchSize), UMLDeformerModel::StaticClass());
+		SettingsCategoryBuilder->AddProperty(GET_MEMBER_NAME_CHECKED(ULegacyVertexDeltaModel, LearningRate), UMLDeformerModel::StaticClass());
+		SettingsCategoryBuilder->AddProperty(GET_MEMBER_NAME_CHECKED(ULegacyVertexDeltaModel, ActivationFunction), UMLDeformerModel::StaticClass());
+		SettingsCategoryBuilder->AddProperty(GET_MEMBER_NAME_CHECKED(ULegacyVertexDeltaModel, LossFunction), UMLDeformerModel::StaticClass());
+		SettingsCategoryBuilder->AddProperty(GET_MEMBER_NAME_CHECKED(ULegacyVertexDeltaModel, ShrinkageSpeed), UMLDeformerModel::StaticClass())
+			.Visibility(TAttribute<EVisibility>::CreateLambda(IsShrinkageVisible));
+		SettingsCategoryBuilder->AddProperty(GET_MEMBER_NAME_CHECKED(ULegacyVertexDeltaModel, ShrinkageThreshold), UMLDeformerModel::StaticClass())
+			.Visibility(TAttribute<EVisibility>::CreateLambda(IsShrinkageVisible));
+		SettingsCategoryBuilder->AddProperty(GET_MEMBER_NAME_CHECKED(ULegacyVertexDeltaModel, MaxCacheSizeGB), UMLDeformerModel::StaticClass());
+	}
+
+	void FLegacyVertexDeltaModelDetails::AddBaseMeshErrors()
+	{
+		// Check if the base mesh matches the target mesh vertex count.
+		FText ErrorText = GetGeomCacheVertexErrorText(Model->SkeletalMesh, VertexModel->GetGeometryCache(), FText::FromString("Base Mesh"), FText::FromString("Target Mesh"));
+		BaseMeshCategoryBuilder->AddCustomRow(FText::FromString("BaseMeshError"))
+			.Visibility(!ErrorText.IsEmpty() ? EVisibility::Visible : EVisibility::Collapsed)
+			.WholeRowContent()
+			[
+				SNew(SBox)
+				.Padding(FMargin(0.0f, 4.0f))
+				[
+					SNew(SWarningOrErrorBox)
+					.MessageStyle(EMessageStyle::Error)
+					.Message(ErrorText)
+				]
+			];
+	}
+	
+	void FLegacyVertexDeltaModelDetails::AddAnimSequenceErrors()
+	{
+		const FText WarningText = GetGeomCacheAnimSequenceErrorText(VertexModel->GetGeometryCache(), Model->GetAnimSequence());
+		BaseMeshCategoryBuilder->AddCustomRow(FText::FromString("AnimSeqWarning"))
+			.Visibility(!WarningText.IsEmpty() ? EVisibility::Visible : EVisibility::Collapsed)
+			.WholeRowContent()
+			[
+				SNew(SBox)
+				.Padding(FMargin(0.0f, 4.0f))
+				[
+					SNew(SWarningOrErrorBox)
+					.MessageStyle(EMessageStyle::Warning)
+					.Message(WarningText)
+				]
+			];
+	}
+
+	void FLegacyVertexDeltaModelDetails::AddTargetMesh()
+	{
+		TargetMeshCategoryBuilder->AddProperty(GET_MEMBER_NAME_CHECKED(ULegacyVertexDeltaModel, GeometryCache), UMLDeformerModel::StaticClass());
+
+		const FText TargetMeshErrorText = GetGeomCacheErrorText(VertexModel->GetSkeletalMesh(), VertexModel->GetGeometryCache());
+		TargetMeshCategoryBuilder->AddCustomRow(FText::FromString("TargetMeshError"))
+			.Visibility(!TargetMeshErrorText.IsEmpty() ? EVisibility::Visible : EVisibility::Collapsed)
+			.WholeRowContent()
+			[
+				SNew(SBox)
+				.Padding(FMargin(0.0f, 4.0f))
+				[
+					SNew(SWarningOrErrorBox)
+					.MessageStyle(EMessageStyle::Error)
+					.Message(TargetMeshErrorText)
+				]
+			];
+
+		const FText ChangedErrorText = EditorModel->GetTargetAssetChangedErrorText();
+		TargetMeshCategoryBuilder->AddCustomRow(FText::FromString("TargetMeshChangedError"))
+			.Visibility(!ChangedErrorText.IsEmpty() ? EVisibility::Visible : EVisibility::Collapsed)
+			.WholeRowContent()
+			[
+				SNew(SBox)
+				.Padding(FMargin(0.0f, 4.0f))
+				[
+					SNew(SWarningOrErrorBox)
+					.MessageStyle(EMessageStyle::Error)
+					.Message(ChangedErrorText)
+				]
+			];
+
+		// Show meshes that have no matching geometry track list.
+		const FText MeshMappingErrorList = GetGeomCacheMeshMappingErrorText(Model->GetSkeletalMesh(), VertexModel->GetGeometryCache());
+		FString GeomTrackNameList;
+		if (!MeshMappingErrorList.IsEmpty())
+		{
+			UGeometryCache* GeomCache = VertexModel->GetGeometryCache();
+			for (int32 Index = 0; Index < GeomCache->Tracks.Num(); ++Index)
+			{
+				GeomTrackNameList += GeomCache->Tracks[Index]->GetName();
+				if (Index < GeomCache->Tracks.Num() - 1)
+				{
+					GeomTrackNameList += TEXT("\n");
+				}
+			}
+		}
+
+		FText MeshMappingErrorFull = FText::Format(
+			LOCTEXT("MeshMappingError", "No matching GeomCache Tracks names found for meshes:\n{0}\n\nGeomCache Track List:\n{1}"), 
+			MeshMappingErrorList,
+			FText::FromString(GeomTrackNameList));
+
+		TargetMeshCategoryBuilder->AddCustomRow(FText::FromString("MeshMappingError"))
+			.Visibility(!MeshMappingErrorList.IsEmpty() ? EVisibility::Visible : EVisibility::Collapsed)
+			.WholeRowContent()
+			[
+				SNew(SBox)
+				.Padding(FMargin(0.0f, 4.0f))
+				[
+					SNew(SWarningOrErrorBox)
+					.MessageStyle(EMessageStyle::Error)
+					.Message(MeshMappingErrorFull)
+				]
+			];
+	}
+}	// namespace UE::LegacyVertexDeltaModel
+
+#undef LOCTEXT_NAMESPACE
diff --git a/Engine/Plugins/Experimental/MLDeformer/LegacyVertexDeltaModel/Source/LegacyVertexDeltaModelEditor/Private/LegacyVertexDeltaModelEditorModule.cpp b/Engine/Plugins/Experimental/MLDeformer/LegacyVertexDeltaModel/Source/LegacyVertexDeltaModelEditor/Private/LegacyVertexDeltaModelEditorModule.cpp
new file mode 100644
--- /dev/null
+++ b/Engine/Plugins/Experimental/MLDeformer/LegacyVertexDeltaModel/Source/LegacyVertexDeltaModelEditor/Private/LegacyVertexDeltaModelEditorModule.cpp
@@ -0,0 +1,56 @@
+// Copyright Epic Games, Inc. All Rights Reserved.
+
+#include "LegacyVertexDeltaModelEditorModule.h"
+#include "MLDeformerEditorModule.h"
+#include "LegacyVertexDeltaModelVizSettingsDetails.h"
+#include "LegacyVertexDeltaModelDetails.h"
+#include "LegacyVertexDeltaEditorModel.h"
+#include "LegacyVertexDeltaModel.h"
+#include "Modules/ModuleManager.h"
+#include "EditorModeRegistry.h"
+
+#define LOCTEXT_NAMESPACE "LegacyVertexDeltaModelEditorModule"
+
+IMPLEMENT_MODULE(UE::LegacyVertexDeltaModel::FLegacyVertexDeltaModelEditorModule, LegacyVertexDeltaModelEditor)
+
+namespace UE::LegacyVertexDeltaModel
+{
+	using namespace UE::MLDeformer;
+
+	void FLegacyVertexDeltaModelEditorModule::StartupModule()
+	{
+		// Register object detail customizations.
+		FPropertyEditorModule& PropertyModule = FModuleManager::LoadModuleChecked<FPropertyEditorModule>("PropertyEditor");
+		PropertyModule.RegisterCustomClassLayout("LegacyVertexDeltaModelVizSettings", FOnGetDetailCustomizationInstance::CreateStatic(&FLegacyVertexDeltaModelVizSettingsDetails::MakeInstance));
+		PropertyModule.RegisterCustomClassLayout("LegacyVertexDeltaModel", FOnGetDetailCustomizationInstance::CreateStatic(&FLegacyVertexDeltaModelDetails::MakeInstance));
+		PropertyModule.NotifyCustomizationModuleChanged();
+
+		// Register our custom ML deformer model to the model registry in the ML Deformer Framework.
+		FMLDeformerEditorModule& EditorModule = FModuleManager::LoadModuleChecked<FMLDeformerEditorModule>("MLDeformerFrameworkEditor");
+		FMLDeformerEditorModelRegistry& ModelRegistry = EditorModule.GetModelRegistry();
+		ModelRegistry.RegisterEditorModel(ULegacyVertexDeltaModel::StaticClass(), FOnGetEditorModelInstance::CreateStatic(&FLegacyVertexDeltaEditorModel::MakeInstance));
+	}
+
+	void FLegacyVertexDeltaModelEditorModule::ShutdownModule()
+	{
+		// Unregister our ML Deformer model.
+		if (FModuleManager::Get().IsModuleLoaded(TEXT("MLDeformerFrameworkEditor")))
+		{
+			FMLDeformerEditorModule& EditorModule = FModuleManager::GetModuleChecked<FMLDeformerEditorModule>("MLDeformerFrameworkEditor");
+			FMLDeformerEditorModelRegistry& ModelRegistry = EditorModule.GetModelRegistry();
+			ModelRegistry.UnregisterEditorModel(ULegacyVertexDeltaModel::StaticClass());
+		}
+
+		// Unregister object detail customizations for this model.
+		if (FModuleManager::Get().IsModuleLoaded(TEXT("PropertyEditor")))
+		{
+			FPropertyEditorModule& PropertyModule = FModuleManager::GetModuleChecked<FPropertyEditorModule>(TEXT("PropertyEditor"));
+			PropertyModule.UnregisterCustomClassLayout(TEXT("LegacyVertexDeltaModelVizSettings"));
+			PropertyModule.UnregisterCustomClassLayout(TEXT("LegacyVertexDeltaModel"));
+			PropertyModule.NotifyCustomizationModuleChanged();
+		}
+	}
+
+}	// namespace UE::LegacyVertexDeltaModel
+
+#undef LOCTEXT_NAMESPACE
diff --git a/Engine/Plugins/Experimental/MLDeformer/LegacyVertexDeltaModel/Source/LegacyVertexDeltaModelEditor/Private/LegacyVertexDeltaModelVizSettingsDetails.cpp b/Engine/Plugins/Experimental/MLDeformer/LegacyVertexDeltaModel/Source/LegacyVertexDeltaModelEditor/Private/LegacyVertexDeltaModelVizSettingsDetails.cpp
new file mode 100644
--- /dev/null
+++ b/Engine/Plugins/Experimental/MLDeformer/LegacyVertexDeltaModel/Source/LegacyVertexDeltaModelEditor/Private/LegacyVertexDeltaModelVizSettingsDetails.cpp
@@ -0,0 +1,104 @@
+// Copyright Epic Games, Inc. All Rights Reserved.
+
+#include "LegacyVertexDeltaModelVizSettingsDetails.h"
+#include "MLDeformerModule.h"
+#include "MLDeformerModel.h"
+#include "MLDeformerEditorModule.h"
+#include "MLDeformerGeomCacheHelpers.h"
+#include "LegacyVertexDeltaModel.h"
+#include "LegacyVertexDeltaEditorModel.h"
+#include "LegacyVertexDeltaModelVizSettings.h"
+#include "DetailLayoutBuilder.h"
+#include "DetailCategoryBuilder.h"
+#include "DetailWidgetRow.h"
+#include "IDetailsView.h"
+#include "IDetailChildrenBuilder.h"
+#include "IDetailGroup.h"
+#include "Modules/ModuleManager.h"
+#include "Layout/Margin.h"
+#include "Widgets/Layout/SBorder.h"
+#include "Widgets/Text/STextBlock.h"
+#include "Widgets/Layout/SBox.h"
+#include "PropertyCustomizationHelpers.h"
+#include "AssetRegistry/AssetData.h"
+#include "PropertyHandle.h"
+#include "DetailLayoutBuilder.h"
+#include "SWarningOrErrorBox.h"
+#include "GeometryCache.h"
+#include "Animation/AnimSequence.h"
+#include "Animation/MeshDeformer.h"
+#include "Engine/SkeletalMesh.h"
+
+#define LOCTEXT_NAMESPACE "LegacyVertexDeltaModelVizSettingsDetails"
+
+namespace UE::LegacyVertexDeltaModel
+{
+	using namespace UE::MLDeformer;
+
+	TSharedRef<IDetailCustomization> FLegacyVertexDeltaModelVizSettingsDetails::MakeInstance()
+	{
+		return MakeShareable(new FLegacyVertexDeltaModelVizSettingsDetails());
+	}
+
+	bool FLegacyVertexDeltaModelVizSettingsDetails::UpdateMemberPointers(const TArray<TWeakObjectPtr<UObject>>& Objects)
+	{
+		if (!FMLDeformerVizSettingsDetails::UpdateMemberPointers(Objects))
+		{
+			return false;
+		}
+
+		VertexDeltaModel = Cast<ULegacyVertexDeltaModel>(Model);
+		VertexDeltaVizSettings = Cast<ULegacyVertexDeltaModelVizSettings>(VizSettings);
+		return (VertexDeltaModel != nullptr && VertexDeltaVizSettings != nullptr);
+	}
+
+	void FLegacyVertexDeltaModelVizSettingsDetails::AddGroundTruth()
+	{
+		TestAssetsCategory->AddProperty(GET_MEMBER_NAME_CHECKED(ULegacyVertexDeltaModelVizSettings, GroundTruth));
+
+		// Show an error when the test anim sequence duration doesn't match the one of the ground truth.
+		const FText AnimErrorText = GetGeomCacheAnimSequenceErrorText(VertexDeltaVizSettings->GetTestGroundTruth(), VizSettings->GetTestAnimSequence());
+		FDetailWidgetRow& GroundTruthAnimErrorRow = TestAssetsCategory->AddCustomRow(FText::FromString("GroundTruthAnimMismatchError"))
+			.Visibility(!AnimErrorText.IsEmpty() ? EVisibility::Visible : EVisibility::Collapsed)
+			.WholeRowContent()
+			[
+				SNew(SBox)
+				.Padding(FMargin(0.0f, 4.0f))
+				[
+					SNew(SWarningOrErrorBox)
+					.MessageStyle(EMessageStyle::Warning)
+					.Message(AnimErrorText)
+				]
+			];
+
+		const FText GeomErrorText = GetGeomCacheErrorText(Model->GetSkeletalMesh(), VertexDeltaVizSettings->GetTestGroundTruth());
+		FDetailWidgetRow& GroundTruthGeomErrorRow = TestAssetsCategory->AddCustomRow(FText::FromString("GroundTruthGeomMismatchError"))
+			.Visibility(!GeomErrorText.IsEmpty() ? EVisibility::Visible : EVisibility::Collapsed)
+			.WholeRowContent()
+			[
+				SNew(SBox)
+				.Padding(FMargin(0.0f, 4.0f))
+				[
+					SNew(SWarningOrErrorBox)
+					.MessageStyle(EMessageStyle::Warning)
+					.Message(GeomErrorText)
+				]
+			];
+
+		const FText VertexErrorText = GetGeomCacheVertexErrorText(Model->GetSkeletalMesh(), VertexDeltaVizSettings->GetTestGroundTruth(), FText::FromString("Base Mesh"), FText::FromString("Ground Truth Mesh"));
+		FDetailWidgetRow& GroundTruthVertexErrorRow = TestAssetsCategory->AddCustomRow(FText::FromString("GroundTruthVertexMismatchError"))
+			.Visibility(!VertexErrorText.IsEmpty() ? EVisibility::Visible : EVisibility::Collapsed)
+			.WholeRowContent()
+			[
+				SNew(SBox)
+				.Padding(FMargin(0.0f, 4.0f))
+				[
+					SNew(SWarningOrErrorBox)
+					.MessageStyle(EMessageStyle::Warning)
+					.Message(VertexErrorText)
+				]
+			];
+	}
+}	//namespace UE::LegacyVertexDeltaModel
+
+#undef LOCTEXT_NAMESPACE
diff --git a/Engine/Plugins/Experimental/MLDeformer/LegacyVertexDeltaModel/Source/LegacyVertexDeltaModelEditor/Private/LegacyVertexDeltaTrainingModel.cpp b/Engine/Plugins/Experimental/MLDeformer/LegacyVertexDeltaModel/Source/LegacyVertexDeltaModelEditor/Private/LegacyVertexDeltaTrainingModel.cpp
new file mode 100644
--- /dev/null
+++ b/Engine/Plugins/Experimental/MLDeformer/LegacyVertexDeltaModel/Source/LegacyVertexDeltaModelEditor/Private/LegacyVertexDeltaTrainingModel.cpp
@@ -0,0 +1,117 @@
+// Copyright Epic Games, Inc. All Rights Reserved.
+
+#include "LegacyVertexDeltaTrainingModel.h"
+#include "LegacyVertexDeltaModel.h"
+#include "LegacyVertexDeltaEditorModel.h"
+#include "MLDeformerEditorModel.h"
+#include "Misc/ScopedSlowTask.h"
+#include "Math/NumericLimits.h"
+
+#define LOCTEXT_NAMESPACE "LegacyVertexDeltaTrainingModel"
+
+ULegacyVertexDeltaModel* ULegacyVertexDeltaTrainingModel::GetVertexDeltaModel() const
+{
+	return Cast<ULegacyVertexDeltaModel>(GetModel());
+}
+
+UE::LegacyVertexDeltaModel::FLegacyVertexDeltaEditorModel* ULegacyVertexDeltaTrainingModel::GetVertexDeltaEditorModel() const
+{
+	return static_cast<UE::LegacyVertexDeltaModel::FLegacyVertexDeltaEditorModel*>(EditorModel);
+}
+
+void ULegacyVertexDeltaTrainingModel::UpdateVertexDeltaMeanAndScale(const TArray<float>& VertexDeltas, FVector& InOutMeanVertexDelta, FVector& InOutVertexDeltaScale, float& InOutCount)
+{
+	// Calculate this frame's mean.
+	FVector MeanDelta = FVector::ZeroVector;
+	FVector MinDelta(TNumericLimits<float>::Max());
+	FVector MaxDelta(-TNumericLimits<float>::Max());
+	const int32 NumVertices = VertexDeltas.Num() / 3;
+	for (int32 VertexIndex = 0; VertexIndex < NumVertices; ++VertexIndex)
+	{
+		const int Offset = VertexIndex * 3;
+		const FVector Delta(VertexDeltas[Offset], VertexDeltas[Offset+1], VertexDeltas[Offset+2]);
+		MeanDelta += Delta;
+		MinDelta = MinDelta.ComponentMin(Delta);
+		MaxDelta = MaxDelta.ComponentMax(Delta);
+	}
+	const float ValidVertexCount = (float)NumVertices;
+	if (NumVertices > 0)
+	{
+		MeanDelta /= (float)NumVertices;
+	}
+
+	// Update global mean.
+	InOutCount += 1.0f;
+	FVector TempDiff = MeanDelta - (FVector)InOutMeanVertexDelta;
+	InOutMeanVertexDelta = InOutMeanVertexDelta + FVector(TempDiff / InOutCount);
+
+	// Update global scale.
+	TempDiff = MaxDelta - MinDelta;
+	InOutVertexDeltaScale = InOutVertexDeltaScale.ComponentMax(TempDiff.GetAbs());
+}
+
+bool ULegacyVertexDeltaTrainingModel::ComputeVertexDeltaStatistics(uint32 LODIndex)
+{
+	using namespace UE::LegacyVertexDeltaModel;
+
+	FLegacyVertexDeltaEditorModel* VertexDeltaEditorModel = GetVertexDeltaEditorModel();
+	ULegacyVertexDeltaModel* VertexDeltaModel = VertexDeltaEditorModel->GetVertexDeltaModel();
+	if (VertexDeltaEditorModel->IsDataNormalized())
+	{
+		return true;
+	}
+
+	FMLDeformerSampler* Sampler = VertexDeltaEditorModel->GetSampler();
+	Sampler->SetVertexDeltaSpace(EVertexDeltaSpace::PreSkinning);
+
+	// Show some dialog with progress.
+	const FText Title(LOCTEXT("PreprocessTrainingDataMessage", "Calculating data statistics"));
+	const uint32 AnimNumFrames = EditorModel->GetNumFramesForTraining();
+	FScopedSlowTask Task((float)AnimNumFrames, Title);
+	Task.MakeDialog(true);
+
+	// Get the base actor data.
+	FVector Mean = FVector::ZeroVector;
+	FVector Scale = FVector::OneVector;
+	float Count = 0.0f;
+	const int32 FrameCount = VertexDeltaEditorModel->GetNumFramesForTraining();
+	for (int32 FrameIndex = 0; FrameIndex < FrameCount; FrameIndex++)
+	{
+		// Calculate the deltas and update the mean and scale.
+		Sampler->Sample(FrameIndex);
+		UpdateVertexDeltaMeanAndScale(Sampler->GetVertexDeltas(), Mean, Scale, Count);
+
+		// Forward the progress bar and check if we want to cancel.
+		Task.EnterProgressFrame();
+		if (Task.ShouldCancel())
+		{
+			return false;
+		}
+	}
+
+	// Update the asset with calculated statistics.
+	VertexDeltaModel->VertexDeltaMean = Mean;
+	if (Count > 0.0f)
+	{
+		VertexDeltaModel->VertexDeltaScale = FVector::OneVector * Scale.GetMax();
+		VertexDeltaEditorModel->SetDataNormalized(true);
+	}
+
+	return true;
+}
+
+bool ULegacyVertexDeltaTrainingModel::ComputeDeltasStatistics()
+{
+	const ULegacyVertexDeltaModel* VertexDeltaModel = GetVertexDeltaModel();
+	const float DeltaCutoffLength = VertexDeltaModel->GetDeltaCutoffLength();
+	if (ComputeVertexDeltaStatistics(0))
+	{
+		VertexDeltaMean = VertexDeltaModel->VertexDeltaMean;
+		VertexDeltaScale = VertexDeltaModel->VertexDeltaScale;
+		return true;
+	}
+
+	return false;
+}
+
+#undef LOCTEXT_NAMESPACE
diff --git a/Engine/Plugins/Experimental/MLDeformer/LegacyVertexDeltaModel/Source/LegacyVertexDeltaModelEditor/Public/LegacyVertexDeltaEditorModel.h b/Engine/Plugins/Experimental/MLDeformer/LegacyVertexDeltaModel/Source/LegacyVertexDeltaModelEditor/Public/LegacyVertexDeltaEditorModel.h
new file mode 100644
--- /dev/null
+++ b/Engine/Plugins/Experimental/MLDeformer/LegacyVertexDeltaModel/Source/LegacyVertexDeltaModelEditor/Public/LegacyVertexDeltaEditorModel.h
@@ -0,0 +1,73 @@
+// Copyright Epic Games, Inc. All Rights Reserved.
+
+#pragma once
+
+#include "CoreMinimal.h"
+#include "MLDeformerEditorModel.h"
+#include "MLDeformerGeomCacheSampler.h"
+#include "LegacyVertexDeltaModel.h"
+#include "MLDeformerEditorActor.h"
+#include "UObject/GCObject.h"
+
+class UMLDeformerAsset;
+class USkeletalMesh;
+class UGeometryCache;
+class ULegacyVertexDeltaModel;
+class ULegacyVertexDeltaModelVizSettings;
+
+namespace UE::LegacyVertexDeltaModel
+{
+	using namespace UE::MLDeformer;
+
+	class FLegacyVertexDeltaEditorModelActor;
+	class FLegacyVertexDeltaModelSampler;
+
+	class LEGACYVERTEXDELTAMODELEDITOR_API FLegacyVertexDeltaEditorModel
+		: public UE::MLDeformer::FMLDeformerEditorModel
+	{
+	public:
+		virtual ~FLegacyVertexDeltaEditorModel() override;
+
+		// We need to implement this static MakeInstance method.
+		static FMLDeformerEditorModel* MakeInstance();
+
+		// FGCObject overrides.
+		virtual void AddReferencedObjects(FReferenceCollector& Collector) override;
+		virtual FString GetReferencerName() const override { return TEXT("FLegacyVertexDeltaEditorModel"); }
+		// ~END FGCObject overrides.
+	
+		// FMLDeformerEditorModel overrides.
+		virtual void Init(const InitSettings& Settings) override;
+		virtual FMLDeformerEditorActor* CreateEditorActor(const FMLDeformerEditorActor::FConstructSettings& Settings) const override;
+		virtual FMLDeformerSampler* CreateSampler() const;
+		virtual void CreateTrainingGroundTruthActor(UWorld* World) override;
+		virtual void CreateTestGroundTruthActor(UWorld* World) override;
+		virtual void OnPreTraining() override;
+		virtual void OnTrainingAborted() override;
+		virtual void OnInputAssetsChanged() override;
+		virtual int32 GetNumFrames() const override;
+		virtual double GetTimeAtFrame(int32 FrameNumber) const override;
+		virtual int32 GetFrameAtTime(double TimeInSeconds) const override;
+		virtual void UpdateIsReadyForTrainingState() override;
+		virtual ETrainingResult Train() override;
+		virtual FString GetDefaultDeformerGraphAssetPath() const override;
+		virtual FString GetHeatMapDeformerGraphPath() const override;
+		FString GetTrainedNetworkOnnxFile() const override;
+		// ~END FMLDeformerEditorModel overrides.
+
+		// Some helpers that cast to this model's variants of some classes.
+		ULegacyVertexDeltaModel* GetVertexDeltaModel() const { return Cast<ULegacyVertexDeltaModel>(Model); }
+		ULegacyVertexDeltaModelVizSettings* GetVertexDeltaModelVizSettings() const;
+		FLegacyVertexDeltaEditorModelActor* FindVertexDeltaModelEditorActor(int32 TypeID) const;
+		FMLDeformerGeomCacheSampler* GetGeomCacheSampler() const { return static_cast<FMLDeformerGeomCacheSampler*>(Sampler); }
+	
+		void OnPostEditChangeProperty(FPropertyChangedEvent& PropertyChangedEvent);
+
+	protected:
+		void CreateGeomCacheActor(UWorld* World, int32 ActorID, const FName& Name, UGeometryCache* GeomCache, FLinearColor LabelColor, FLinearColor WireframeColor, const FText& LabelText, bool bIsTrainingActor);
+
+	protected:
+		FVector VertexDeltaMeanBackup = FVector::ZeroVector;
+		FVector VertexDeltaScaleBackup = FVector::OneVector;
+	};
+}	// namespace UE::LegacyVertexDeltaModel
diff --git a/Engine/Plugins/Experimental/MLDeformer/LegacyVertexDeltaModel/Source/LegacyVertexDeltaModelEditor/Public/LegacyVertexDeltaEditorModelActor.h b/Engine/Plugins/Experimental/MLDeformer/LegacyVertexDeltaModel/Source/LegacyVertexDeltaModelEditor/Public/LegacyVertexDeltaEditorModelActor.h
new file mode 100644
--- /dev/null
+++ b/Engine/Plugins/Experimental/MLDeformer/LegacyVertexDeltaModel/Source/LegacyVertexDeltaModelEditor/Public/LegacyVertexDeltaEditorModelActor.h
@@ -0,0 +1,37 @@
+// Copyright Epic Games, Inc. All Rights Reserved.
+
+#pragma once
+
+#include "MLDeformerEditorActor.h"
+
+class UMLDeformerComponent;
+class UGeometryCacheComponent;
+
+namespace UE::LegacyVertexDeltaModel
+{
+	using namespace UE::MLDeformer;
+
+	class LEGACYVERTEXDELTAMODELEDITOR_API FLegacyVertexDeltaEditorModelActor
+		: public FMLDeformerEditorActor
+	{
+	public:
+		FLegacyVertexDeltaEditorModelActor(const FConstructSettings& Settings);
+
+		void SetGeometryCacheComponent(UGeometryCacheComponent* Component) { GeomCacheComponent = Component; }
+		UGeometryCacheComponent* GetGeometryCacheComponent() const { return GeomCacheComponent; }
+
+		// FMLDeformerEditorActor overrides.
+		virtual void SetVisibility(bool bIsVisible) override;
+		virtual bool IsVisible() const override;
+		virtual void SetPlayPosition(float TimeInSeconds, bool bAutoPause=true) override;
+		virtual float GetPlayPosition() const override;
+		virtual void SetPlaySpeed(float PlaySpeed) override;
+		virtual void Pause(bool bPaused) override;
+		virtual FBox GetBoundingBox() const override;
+		// ~END FMLDeformerEditorActor overrides.
+
+	protected:
+		/** The geometry cache component (can be nullptr). */
+		TObjectPtr<UGeometryCacheComponent> GeomCacheComponent = nullptr;
+	};
+}	// namespace UE::LegacyVertexDeltaModel
diff --git a/Engine/Plugins/Experimental/MLDeformer/LegacyVertexDeltaModel/Source/LegacyVertexDeltaModelEditor/Public/LegacyVertexDeltaModelDetails.h b/Engine/Plugins/Experimental/MLDeformer/LegacyVertexDeltaModel/Source/LegacyVertexDeltaModelEditor/Public/LegacyVertexDeltaModelDetails.h
new file mode 100644
--- /dev/null
+++ b/Engine/Plugins/Experimental/MLDeformer/LegacyVertexDeltaModel/Source/LegacyVertexDeltaModelEditor/Public/LegacyVertexDeltaModelDetails.h
@@ -0,0 +1,45 @@
+// Copyright Epic Games, Inc. All Rights Reserved.
+
+#pragma once
+
+#include "CoreMinimal.h"
+#include "MLDeformerModelDetails.h"
+#include "IDetailCustomization.h"
+#include "AssetRegistry/AssetData.h"
+#include "Input/Reply.h"
+
+class IDetailLayoutBuilder;
+class ULegacyVertexDeltaModel;
+
+namespace UE::MLDeformer
+{
+	class FMLDeformerEditorModel;
+}
+
+namespace UE::LegacyVertexDeltaModel
+{
+	class FLegacyVertexDeltaEditorModel;
+
+	class LEGACYVERTEXDELTAMODELEDITOR_API FLegacyVertexDeltaModelDetails final
+		: public UE::MLDeformer::FMLDeformerModelDetails
+	{
+	public:
+		/** Makes a new instance of this detail layout class for a specific detail view requesting it. */
+		static TSharedRef<IDetailCustomization> MakeInstance();
+
+		// ILayoutDetails overrides.
+		void CustomizeDetails(IDetailLayoutBuilder& DetailBuilder) override;
+		// ~END ILayoutDetails overrides.
+
+		// FMLDeformerModelDetails overrides.
+		bool UpdateMemberPointers(const TArray<TWeakObjectPtr<UObject>>& Objects) override;
+		void AddTargetMesh() override;
+		void AddBaseMeshErrors() override;
+		void AddAnimSequenceErrors() override;
+		// ~END FMLDeformerModelDetails overrides.
+
+	protected:
+		TObjectPtr<ULegacyVertexDeltaModel> VertexModel = nullptr;
+		FLegacyVertexDeltaEditorModel* VertexEditorModel = nullptr;
+	};
+}	// namespace UE::LegacyVertexDeltaModel
diff --git a/Engine/Plugins/Experimental/MLDeformer/LegacyVertexDeltaModel/Source/LegacyVertexDeltaModelEditor/Public/LegacyVertexDeltaModelEditorModule.h b/Engine/Plugins/Experimental/MLDeformer/LegacyVertexDeltaModel/Source/LegacyVertexDeltaModelEditor/Public/LegacyVertexDeltaModelEditorModule.h
new file mode 100644
--- /dev/null
+++ b/Engine/Plugins/Experimental/MLDeformer/LegacyVertexDeltaModel/Source/LegacyVertexDeltaModelEditor/Public/LegacyVertexDeltaModelEditorModule.h
@@ -0,0 +1,19 @@
+// Copyright Epic Games, Inc. All Rights Reserved.
+
+#pragma once
+
+#include "CoreMinimal.h"
+#include "Modules/ModuleInterface.h"
+
+namespace UE::LegacyVertexDeltaModel
+{
+	class LEGACYVERTEXDELTAMODELEDITOR_API FLegacyVertexDeltaModelEditorModule
+		: public IModuleInterface
+	{
+	public:
+		// IModuleInterface overrides.
+		void StartupModule() override;
+		void ShutdownModule() override;
+		// ~END IModuleInterface overrides.
+	};
+}	// namespace UE::LegacyVertexDeltaModel
\ No newline at end of file
diff --git a/Engine/Plugins/Experimental/MLDeformer/LegacyVertexDeltaModel/Source/LegacyVertexDeltaModelEditor/Public/LegacyVertexDeltaModelVizSettingsDetails.h b/Engine/Plugins/Experimental/MLDeformer/LegacyVertexDeltaModel/Source/LegacyVertexDeltaModelEditor/Public/LegacyVertexDeltaModelVizSettingsDetails.h
new file mode 100644
--- /dev/null
+++ b/Engine/Plugins/Experimental/MLDeformer/LegacyVertexDeltaModel/Source/LegacyVertexDeltaModelEditor/Public/LegacyVertexDeltaModelVizSettingsDetails.h
@@ -0,0 +1,34 @@
+// Copyright Epic Games, Inc. All Rights Reserved.
+
+#pragma once
+
+#include "CoreMinimal.h"
+#include "MLDeformerVizSettingsDetails.h"
+#include "IDetailCustomization.h"
+#include "AssetRegistry/AssetData.h"
+#include "PropertyHandle.h"
+
+class IDetailLayoutBuilder;
+class USkeleton;
+class ULegacyVertexDeltaModel;
+class ULegacyVertexDeltaModelVizSettings;
+
+namespace UE::LegacyVertexDeltaModel
+{
+	class LEGACYVERTEXDELTAMODELEDITOR_API FLegacyVertexDeltaModelVizSettingsDetails final
+		: public UE::MLDeformer::FMLDeformerVizSettingsDetails
+	{
+	public:
+		/** Makes a new instance of this detail layout class for a specific detail view requesting it. */
+		static TSharedRef<IDetailCustomization> MakeInstance();
+
+		// FMLDeformerVizSettingsDetails overrides.
+		bool UpdateMemberPointers(const TArray<TWeakObjectPtr<UObject>>& Objects) override;
+		void AddGroundTruth() override;
+		// ~END FMLDeformerVizSettingsDetails overrides.
+
+	protected:
+		TObjectPtr<ULegacyVertexDeltaModel> VertexDeltaModel = nullptr;
+		TObjectPtr<ULegacyVertexDeltaModelVizSettings> VertexDeltaVizSettings = nullptr;
+	};
+}	// namespace UE::LegacyVertexDeltaModel
diff --git a/Engine/Plugins/Experimental/MLDeformer/LegacyVertexDeltaModel/Source/LegacyVertexDeltaModelEditor/Public/LegacyVertexDeltaTrainingModel.h b/Engine/Plugins/Experimental/MLDeformer/LegacyVertexDeltaModel/Source/LegacyVertexDeltaModelEditor/Public/LegacyVertexDeltaTrainingModel.h
new file mode 100644
--- /dev/null
+++ b/Engine/Plugins/Experimental/MLDeformer/LegacyVertexDeltaModel/Source/LegacyVertexDeltaModelEditor/Public/LegacyVertexDeltaTrainingModel.h
@@ -0,0 +1,46 @@
+// Copyright Epic Games, Inc. All Rights Reserved.
+
+#pragma once
+
+#include "CoreMinimal.h"
+#include "MLDeformerTrainingModel.h"
+#include "LegacyVertexDeltaTrainingModel.generated.h"
+
+namespace UE::LegacyVertexDeltaModel
+{
+	class FLegacyVertexDeltaEditorModel;
+}
+
+class ULegacyVertexDeltaModel;
+
+UCLASS(Blueprintable)
+class LEGACYVERTEXDELTAMODELEDITOR_API ULegacyVertexDeltaTrainingModel
+	: public UMLDeformerTrainingModel
+{
+	GENERATED_BODY()
+
+public:
+	/** Main training function, with implementation in python. */
+	UFUNCTION(BlueprintImplementableEvent, Category = "Python")
+	int32 Train() const;
+
+	/** Compute delta statistics for the whole dataset. */
+	UFUNCTION(BlueprintCallable, Category = "Training Data")
+	bool ComputeDeltasStatistics();
+
+	ULegacyVertexDeltaModel* GetVertexDeltaModel() const;
+	UE::LegacyVertexDeltaModel::FLegacyVertexDeltaEditorModel* GetVertexDeltaEditorModel() const;
+
+protected:
+	bool ComputeVertexDeltaStatistics(uint32 LODIndex);
+	void UpdateVertexDeltaMeanAndScale(const TArray<float>& VertexDeltas, FVector& InOutMeanVertexDelta, FVector& InOutVertexDeltaScale, float& InOutCount);
+
+public:
+	// Mean delta computed over the entire dataset.
+	UPROPERTY(EditAnywhere, BlueprintReadWrite, Category = "Training Data")
+	FVector VertexDeltaMean = FVector::ZeroVector;
+
+	// Vertex delta scale computed over the entire dataset.
+	UPROPERTY(EditAnywhere, BlueprintReadWrite, Category = "Training Data")
+	FVector VertexDeltaScale = FVector::OneVector;
+};
diff --git a/Engine/Plugins/Experimental/MLDeformer/MLDeformerFramework/Source/MLDeformerFramework/Public/MLDeformerModel.h b/Engine/Plugins/Experimental/MLDeformer/MLDeformerFramework/Source/MLDeformerFramework/Public/MLDeformerModel.h
--- a/Engine/Plugins/Experimental/MLDeformer/MLDeformerFramework/Source/MLDeformerFramework/Public/MLDeformerModel.h
+++ b/Engine/Plugins/Experimental/MLDeformer/MLDeformerFramework/Source/MLDeformerFramework/Public/MLDeformerModel.h
@@ -119,7 +119,6 @@ public:
 	const UAnimSequence* GetAnimSequence() const { return AnimSequence;  }
 	UAnimSequence* GetAnimSequence() { return AnimSequence; }
 	int32 GetTrainingFrameLimit() const { return MaxTrainingFrames; }
-	int32 GetMaxCacheSizeInGigaBytes() const { return MaxCacheSizeGB; }
 	const FTransform& GetAlignmentTransform() const { return AlignmentTransform; }
 	TArray<FBoneReference>& GetBoneIncludeList() { return BoneIncludeList; }
 	const TArray<FBoneReference>& GetBoneIncludeList() const { return BoneIncludeList; }
@@ -202,13 +201,5 @@ public:
 	  * Deltas that are longer than the cutoff value (in units), will be ignored and set to zero length. */
 	UPROPERTY(EditAnywhere, BlueprintReadWrite, Category = "Inputs and Output", meta = (ClampMin = "0.01", ForceUnits="cm"))
 	float DeltaCutoffLength = 30.0f;
-
-	/** 
-	 * The maximum allowed size of the training cache in memory, in gigabytes.
-	 * So a value of 4 would use a maximum of four gigabyte of system memory. 
-	 * The larger the cache size the faster the training.
-	 */
-	UPROPERTY(EditAnywhere, BlueprintReadWrite, Category = "Training Settings", DisplayName = "Max Cache Size", meta = (ClampMin = "0", ForceUnits = "Gigabytes"))
-	int32 MaxCacheSizeGB = 4;	// 4 Gigabyte
 #endif
 };
diff --git a/Engine/Plugins/Experimental/MLDeformer/MLDeformerFramework/Source/MLDeformerFrameworkEditor/Private/MLDeformerModelDetails.cpp b/Engine/Plugins/Experimental/MLDeformer/MLDeformerFramework/Source/MLDeformerFrameworkEditor/Private/MLDeformerModelDetails.cpp
--- a/Engine/Plugins/Experimental/MLDeformer/MLDeformerFramework/Source/MLDeformerFrameworkEditor/Private/MLDeformerModelDetails.cpp
+++ b/Engine/Plugins/Experimental/MLDeformer/MLDeformerFramework/Source/MLDeformerFrameworkEditor/Private/MLDeformerModelDetails.cpp
@@ -215,10 +215,6 @@ namespace UE::MLDeformer
 
 		AddCurveInputErrors();
 
-		SettingsCategoryBuilder->AddProperty(GET_MEMBER_NAME_CHECKED(UMLDeformerModel, MaxCacheSizeGB), UMLDeformerModel::StaticClass());
-
-		AddCacheSizeErrors();
-
 		// Show a warning when no neural network has been set.
 		{		
 			UNeuralNetwork* NeuralNetwork = Model->GetNeuralNetwork();
diff --git a/Engine/Plugins/Experimental/MLDeformer/MLDeformerFramework/Source/MLDeformerFrameworkEditor/Public/MLDeformerModelDetails.h b/Engine/Plugins/Experimental/MLDeformer/MLDeformerFramework/Source/MLDeformerFrameworkEditor/Public/MLDeformerModelDetails.h
--- a/Engine/Plugins/Experimental/MLDeformer/MLDeformerFramework/Source/MLDeformerFrameworkEditor/Public/MLDeformerModelDetails.h
+++ b/Engine/Plugins/Experimental/MLDeformer/MLDeformerFramework/Source/MLDeformerFrameworkEditor/Public/MLDeformerModelDetails.h
@@ -32,7 +32,6 @@ namespace UE::MLDeformer
 		virtual void AddTrainingInputErrors() {}
 		virtual void AddBoneInputErrors() {}
 		virtual void AddCurveInputErrors() {}
-		virtual void AddCacheSizeErrors() {}
 
 	protected:
 		bool FilterAnimSequences(const FAssetData& AssetData, USkeleton* Skeleton);
diff --git a/Engine/Plugins/Experimental/MLDeformer/VertexDeltaModel/Shaders/Private/VertexDeltaModelDataInterface.ush b/Engine/Plugins/Experimental/MLDeformer/VertexDeltaModel/Shaders/Private/VertexDeltaModelDataInterface.ush
deleted file mode 100644
diff --git a/Engine/Plugins/Experimental/MLDeformer/VertexDeltaModel/Source/VertexDeltaModel/Private/Tests/VertexDeltaModelTest.cpp b/Engine/Plugins/Experimental/MLDeformer/VertexDeltaModel/Source/VertexDeltaModel/Private/Tests/VertexDeltaModelTest.cpp
--- a/Engine/Plugins/Experimental/MLDeformer/VertexDeltaModel/Source/VertexDeltaModel/Private/Tests/VertexDeltaModelTest.cpp
+++ b/Engine/Plugins/Experimental/MLDeformer/VertexDeltaModel/Source/VertexDeltaModel/Private/Tests/VertexDeltaModelTest.cpp
@@ -80,9 +80,6 @@ namespace UE::VertexDeltaModelTests
 		const int64 NumNetworkOutputs = NeuralNet->GetOutputTensor().Num();
 		const int64 ExpectedNetworkOutput = NumGeomCacheVerts * 3;	// 3 floats per vertex.
 		UTEST_EQUAL(TEXT("NeuralNet output size check"), NumNetworkOutputs, ExpectedNetworkOutput);
-
-		UTEST_EQUAL(TEXT("Loss function check"), VertexDeltaModel->GetLossFunction(), EVertexDeltaModelLossFunction::MSE);
-		UTEST_EQUAL(TEXT("Activation function check"), VertexDeltaModel->GetActivationFunction(), EVertexDeltaModelActivationFunction::LRelu);
 		UTEST_EQUAL(TEXT("Hidden Layers check"), VertexDeltaModel->GetNumHiddenLayers(), 2);
 
 		UVertexDeltaModelVizSettings* VizSettings = Cast<UVertexDeltaModelVizSettings>(VertexDeltaModel->GetVizSettings());
diff --git a/Engine/Plugins/Experimental/MLDeformer/VertexDeltaModel/Source/VertexDeltaModel/Private/VertexDeltaModelDataInterface.cpp b/Engine/Plugins/Experimental/MLDeformer/VertexDeltaModel/Source/VertexDeltaModel/Private/VertexDeltaModelDataInterface.cpp
deleted file mode 100644
diff --git a/Engine/Plugins/Experimental/MLDeformer/VertexDeltaModel/Source/VertexDeltaModel/Private/VertexDeltaModelModule.cpp b/Engine/Plugins/Experimental/MLDeformer/VertexDeltaModel/Source/VertexDeltaModel/Private/VertexDeltaModelModule.cpp
--- a/Engine/Plugins/Experimental/MLDeformer/VertexDeltaModel/Source/VertexDeltaModel/Private/VertexDeltaModelModule.cpp
+++ b/Engine/Plugins/Experimental/MLDeformer/VertexDeltaModel/Source/VertexDeltaModel/Private/VertexDeltaModelModule.cpp
@@ -14,12 +14,6 @@ IMPLEMENT_MODULE(UE::VertexDeltaModel::FVertexDeltaModelModule, VertexDeltaModel
 
 namespace UE::VertexDeltaModel
 {
-	void FVertexDeltaModelModule::StartupModule()
-	{
-		// Register an additional shader path for our shaders used inside the deformer graph system.
-		const FString PluginShaderDir = FPaths::Combine(IPluginManager::Get().FindPlugin(TEXT("VertexDeltaModel"))->GetBaseDir(), TEXT("Shaders"));
-		AddShaderSourceDirectoryMapping(TEXT("/Plugin/VertexDeltaModel"), PluginShaderDir);
-	}
 }	// namespace UE::VertexDeltaModel
 
 #undef LOCTEXT_NAMESPACE
diff --git a/Engine/Plugins/Experimental/MLDeformer/VertexDeltaModel/Source/VertexDeltaModel/Public/VertexDeltaModel.h b/Engine/Plugins/Experimental/MLDeformer/VertexDeltaModel/Source/VertexDeltaModel/Public/VertexDeltaModel.h
--- a/Engine/Plugins/Experimental/MLDeformer/VertexDeltaModel/Source/VertexDeltaModel/Public/VertexDeltaModel.h
+++ b/Engine/Plugins/Experimental/MLDeformer/VertexDeltaModel/Source/VertexDeltaModel/Public/VertexDeltaModel.h
@@ -17,23 +17,6 @@ class UMLDeformerAsset;
 class USkeleton;
 class IPropertyHandle;
 
-/** The activation function to use during the ML Deformer training process. */
-UENUM()
-enum class EVertexDeltaModelActivationFunction : uint8
-{
-	Relu,
-	LRelu,
-	Tanh
-};
-
-/** The loss function to use during the ML Deformer training process. */
-UENUM()
-enum class EVertexDeltaModelLossFunction : uint8
-{
-	L1,
-	MSE,
-	Shrinkage
-};
 
 UCLASS()
 class VERTEXDELTAMODEL_API UVertexDeltaModel 
@@ -55,33 +38,19 @@ public:
 #endif
 	// ~END UMLDeformerModel overrides.
 
-	const FVector& GetVertexDeltaMean() const { return VertexDeltaMean; }
-	const FVector& GetVertexDeltaScale() const { return VertexDeltaScale; }
-
 #if WITH_EDITORONLY_DATA
 	const UGeometryCache* GetGeometryCache() const { return GeometryCache; }
 	UGeometryCache* GetGeometryCache() { return GeometryCache; }
-	EVertexDeltaModelLossFunction GetLossFunction() const { return LossFunction; }
-	EVertexDeltaModelActivationFunction GetActivationFunction() const { return ActivationFunction; }
 	int32 GetNumHiddenLayers() const { return NumHiddenLayers; }
 	int32 GetNumNeuronsPerLayer() const { return NumNeuronsPerLayer; }
-	int32 GetNumEpochs() const { return Epochs; }
+	int32 GetNumIerations() const { return NumIterations; }
 	int32 GetBatchSize() const { return BatchSize; }
 	float GetLearningRate() const { return LearningRate; }
-	float GetShrinkageSpeed() const { return ShrinkageSpeed; }
-	float GetShrinkageThrehsold() const { return ShrinkageThreshold; }
 	TArray<UE::MLDeformer::FMLDeformerGeomCacheMeshMapping>& GetGeomCacheMeshMappings() { return MeshMappings; }
 	const TArray<UE::MLDeformer::FMLDeformerGeomCacheMeshMapping>& GetGeomCacheMeshMappings() const { return MeshMappings; }
 #endif
 
 public:
-	/** The vertex delta mean. This is passed to the deformer shader to reconstruct the correct output deltas. */
-	UPROPERTY()
-	FVector VertexDeltaMean = FVector::ZeroVector;
-
-	/** The vertex delta scale. This is passed to the deformer shader to rescale the output deltas. */
-	UPROPERTY()
-	FVector VertexDeltaScale = FVector::OneVector;
 
 #if WITH_EDITORONLY_DATA
 	TArray<UE::MLDeformer::FMLDeformerGeomCacheMeshMapping> MeshMappings;
@@ -92,38 +61,23 @@ public:
 
 	/** The number of hidden layers that the neural network model will have.\nHigher numbers will slow down performance but can deal with more complex deformations. */
 	UPROPERTY(EditAnywhere, BlueprintReadWrite, Category = "Training Settings", meta = (ClampMin = "1", ClampMax = "10"))
-	int32 NumHiddenLayers = 2;
+	int32 NumHiddenLayers = 3;
 
 	/** The number of units/neurons per hidden layer. Higher numbers will slow down performance but allow for more complex mesh deformations. */
 	UPROPERTY(EditAnywhere, BlueprintReadWrite, Category = "Training Settings", meta = (ClampMin = "1"))
 	int32 NumNeuronsPerLayer = 256;
 
-	/** The number of epochs to process without any decay. */
+	/** The number of iterations to train the model for. */
 	UPROPERTY(EditAnywhere, BlueprintReadWrite, Category = "Training Settings", meta = (ClampMin = "1"))
-	int32 Epochs = 20;
+	int32 NumIterations = 500000;
 
 	/** The number of frames per batch when training the model. */
-	UPROPERTY(EditAnywhere, BlueprintReadWrite, AdvancedDisplay, Category = "Training Settings", meta = (ClampMin = "1"))
+	UPROPERTY(EditAnywhere, BlueprintReadWrite, Category = "Training Settings", meta = (ClampMin = "1"))
 	int32 BatchSize = 128;
 
 	/** The learning rate used during the model training. */
-	UPROPERTY(EditAnywhere, BlueprintReadWrite, AdvancedDisplay, Category = "Training Settings", meta = (ClampMin = "0.000001", ClampMax = "1.0"))
-	float LearningRate = 0.00175f;
-
-	/** The activation function to use in the neural network. */
-	UPROPERTY(EditAnywhere, BlueprintReadWrite, AdvancedDisplay, Category = "Training Settings")
-	EVertexDeltaModelActivationFunction ActivationFunction = EVertexDeltaModelActivationFunction::LRelu;
-
-	/** The loss function to use during model training. */
-	UPROPERTY(EditAnywhere, BlueprintReadWrite, AdvancedDisplay, Category = "Training Settings")
-	EVertexDeltaModelLossFunction LossFunction = EVertexDeltaModelLossFunction::L1;
-
-	/** Shrinkage speed. Only if the shrinkage loss is used. */
-	UPROPERTY(EditAnywhere, BlueprintReadWrite, AdvancedDisplay, Category = "Training Settings", meta = (ClampMin = "1.0", ClampMax = "1000.0"))
-	float ShrinkageSpeed = 10.0f;
+	UPROPERTY(EditAnywhere, BlueprintReadWrite, Category = "Training Settings", meta = (ClampMin = "0.000001", ClampMax = "1.0"))
+	float LearningRate = 0.001f;
 
-	/** Shrinkage threshold. Only if the shrinkage loss is used. */
-	UPROPERTY(EditAnywhere, BlueprintReadWrite, AdvancedDisplay, Category = "Training Settings", meta = (ClampMin = "0.0"))
-	float ShrinkageThreshold = 0.1f;
 #endif // WITH_EDITORONLY_DATA
 };
diff --git a/Engine/Plugins/Experimental/MLDeformer/VertexDeltaModel/Source/VertexDeltaModel/Public/VertexDeltaModelDataInterface.h b/Engine/Plugins/Experimental/MLDeformer/VertexDeltaModel/Source/VertexDeltaModel/Public/VertexDeltaModelDataInterface.h
deleted file mode 100644
diff --git a/Engine/Plugins/Experimental/MLDeformer/VertexDeltaModel/Source/VertexDeltaModel/Public/VertexDeltaModelModule.h b/Engine/Plugins/Experimental/MLDeformer/VertexDeltaModel/Source/VertexDeltaModel/Public/VertexDeltaModelModule.h
--- a/Engine/Plugins/Experimental/MLDeformer/VertexDeltaModel/Source/VertexDeltaModel/Public/VertexDeltaModelModule.h
+++ b/Engine/Plugins/Experimental/MLDeformer/VertexDeltaModel/Source/VertexDeltaModel/Public/VertexDeltaModelModule.h
@@ -12,9 +12,5 @@ namespace UE::VertexDeltaModel
 	class VERTEXDELTAMODEL_API FVertexDeltaModelModule
 		: public IModuleInterface
 	{
-	public:
-		// IModuleInterface overrides.
-		virtual void StartupModule() override;
-		// ~END IModuleInterface overrides.
 	};
 }	// namespace VertexDeltaModel
diff --git a/Engine/Plugins/Experimental/MLDeformer/VertexDeltaModel/Source/VertexDeltaModelEditor/Private/VertexDeltaEditorModel.cpp b/Engine/Plugins/Experimental/MLDeformer/VertexDeltaModel/Source/VertexDeltaModelEditor/Private/VertexDeltaEditorModel.cpp
--- a/Engine/Plugins/Experimental/MLDeformer/VertexDeltaModel/Source/VertexDeltaModelEditor/Private/VertexDeltaEditorModel.cpp
+++ b/Engine/Plugins/Experimental/MLDeformer/VertexDeltaModel/Source/VertexDeltaModelEditor/Private/VertexDeltaEditorModel.cpp
@@ -269,37 +269,10 @@ namespace UE::VertexDeltaModel
 		return TrainModel<UVertexDeltaTrainingModel>(this);
 	}
 
-	void FVertexDeltaEditorModel::OnPreTraining()
-	{
-		// Make a backup of the normalization values, as they get overwritten when training.
-		// However, when we abort, we want to restore to the original values again.
-		// See OnTrainingAborted for when we restore the backup again.
-		VertexDeltaMeanBackup = GetVertexDeltaModel()->GetVertexDeltaMean();
-		VertexDeltaScaleBackup = GetVertexDeltaModel()->GetVertexDeltaScale();
-	}
-
-	void FVertexDeltaEditorModel::OnTrainingAborted()
-	{
-		// Restore the vertex delta mean and scale, as we aborted, and they could have changed when training
-		// on a smaller subset of frames/samples. If we don't do this, the mesh will deform incorrectly.
-		GetVertexDeltaModel()->VertexDeltaMean = VertexDeltaMeanBackup;
-		GetVertexDeltaModel()->VertexDeltaScale = VertexDeltaScaleBackup;
-	}
-
 	FString FVertexDeltaEditorModel::GetTrainedNetworkOnnxFile() const
 	{
 		return FString(FPaths::ProjectIntermediateDir() + TEXT("VertexDeltaModel/latest_net_G.onnx"));
 	}
-
-	FString FVertexDeltaEditorModel::GetDefaultDeformerGraphAssetPath() const
-	{
-		return FString(TEXT("/VertexDeltaModel/Deformers/DG_VertexDeltaModel.DG_VertexDeltaModel"));
-	}
-
-	FString FVertexDeltaEditorModel::GetHeatMapDeformerGraphPath() const
-	{
-		return FString(TEXT("/VertexDeltaModel/Deformers/DG_VertexDeltaModel_HeatMap.DG_VertexDeltaModel_HeatMap"));
-	}
 }	// namespace UE::VertexDeltaModel
 
 #undef LOCTEXT_NAMESPACE
diff --git a/Engine/Plugins/Experimental/MLDeformer/VertexDeltaModel/Source/VertexDeltaModelEditor/Private/VertexDeltaModelDetails.cpp b/Engine/Plugins/Experimental/MLDeformer/VertexDeltaModel/Source/VertexDeltaModelEditor/Private/VertexDeltaModelDetails.cpp
--- a/Engine/Plugins/Experimental/MLDeformer/VertexDeltaModel/Source/VertexDeltaModelEditor/Private/VertexDeltaModelDetails.cpp
+++ b/Engine/Plugins/Experimental/MLDeformer/VertexDeltaModel/Source/VertexDeltaModelEditor/Private/VertexDeltaModelDetails.cpp
@@ -61,23 +61,9 @@ namespace UE::VertexDeltaModel
 		// Training settings.
 		SettingsCategoryBuilder->AddProperty(GET_MEMBER_NAME_CHECKED(UVertexDeltaModel, NumHiddenLayers), UMLDeformerModel::StaticClass());
 		SettingsCategoryBuilder->AddProperty(GET_MEMBER_NAME_CHECKED(UVertexDeltaModel, NumNeuronsPerLayer), UMLDeformerModel::StaticClass());
-		SettingsCategoryBuilder->AddProperty(GET_MEMBER_NAME_CHECKED(UVertexDeltaModel, Epochs), UMLDeformerModel::StaticClass());
-
-		// Check whether shrinkage settings should be visible or not.
-		auto IsShrinkageVisible = [this]()
-		{
-			return (VertexModel->GetLossFunction() == EVertexDeltaModelLossFunction::Shrinkage) ? EVisibility::Visible : EVisibility::Collapsed;
-		};
-	 
-		// Advanced settings.
+		SettingsCategoryBuilder->AddProperty(GET_MEMBER_NAME_CHECKED(UVertexDeltaModel, NumIterations), UMLDeformerModel::StaticClass());
 		SettingsCategoryBuilder->AddProperty(GET_MEMBER_NAME_CHECKED(UVertexDeltaModel, BatchSize), UMLDeformerModel::StaticClass());
 		SettingsCategoryBuilder->AddProperty(GET_MEMBER_NAME_CHECKED(UVertexDeltaModel, LearningRate), UMLDeformerModel::StaticClass());
-		SettingsCategoryBuilder->AddProperty(GET_MEMBER_NAME_CHECKED(UVertexDeltaModel, ActivationFunction), UMLDeformerModel::StaticClass());
-		SettingsCategoryBuilder->AddProperty(GET_MEMBER_NAME_CHECKED(UVertexDeltaModel, LossFunction), UMLDeformerModel::StaticClass());
-		SettingsCategoryBuilder->AddProperty(GET_MEMBER_NAME_CHECKED(UVertexDeltaModel, ShrinkageSpeed), UMLDeformerModel::StaticClass())
-			.Visibility(TAttribute<EVisibility>::CreateLambda(IsShrinkageVisible));
-		SettingsCategoryBuilder->AddProperty(GET_MEMBER_NAME_CHECKED(UVertexDeltaModel, ShrinkageThreshold), UMLDeformerModel::StaticClass())
-			.Visibility(TAttribute<EVisibility>::CreateLambda(IsShrinkageVisible));
 	}
 
 	void FVertexDeltaModelDetails::AddBaseMeshErrors()
diff --git a/Engine/Plugins/Experimental/MLDeformer/VertexDeltaModel/Source/VertexDeltaModelEditor/Private/VertexDeltaTrainingModel.cpp b/Engine/Plugins/Experimental/MLDeformer/VertexDeltaModel/Source/VertexDeltaModelEditor/Private/VertexDeltaTrainingModel.cpp
--- a/Engine/Plugins/Experimental/MLDeformer/VertexDeltaModel/Source/VertexDeltaModelEditor/Private/VertexDeltaTrainingModel.cpp
+++ b/Engine/Plugins/Experimental/MLDeformer/VertexDeltaModel/Source/VertexDeltaModelEditor/Private/VertexDeltaTrainingModel.cpp
@@ -19,99 +19,4 @@ UE::VertexDeltaModel::FVertexDeltaEditorModel* UVertexDeltaTrainingModel::GetVer
 	return static_cast<UE::VertexDeltaModel::FVertexDeltaEditorModel*>(EditorModel);
 }
 
-void UVertexDeltaTrainingModel::UpdateVertexDeltaMeanAndScale(const TArray<float>& VertexDeltas, FVector& InOutMeanVertexDelta, FVector& InOutVertexDeltaScale, float& InOutCount)
-{
-	// Calculate this frame's mean.
-	FVector MeanDelta = FVector::ZeroVector;
-	FVector MinDelta(TNumericLimits<float>::Max());
-	FVector MaxDelta(-TNumericLimits<float>::Max());
-	const int32 NumVertices = VertexDeltas.Num() / 3;
-	for (int32 VertexIndex = 0; VertexIndex < NumVertices; ++VertexIndex)
-	{
-		const int Offset = VertexIndex * 3;
-		const FVector Delta(VertexDeltas[Offset], VertexDeltas[Offset+1], VertexDeltas[Offset+2]);
-		MeanDelta += Delta;
-		MinDelta = MinDelta.ComponentMin(Delta);
-		MaxDelta = MaxDelta.ComponentMax(Delta);
-	}
-	const float ValidVertexCount = (float)NumVertices;
-	if (NumVertices > 0)
-	{
-		MeanDelta /= (float)NumVertices;
-	}
-
-	// Update global mean.
-	InOutCount += 1.0f;
-	FVector TempDiff = MeanDelta - (FVector)InOutMeanVertexDelta;
-	InOutMeanVertexDelta = InOutMeanVertexDelta + FVector(TempDiff / InOutCount);
-
-	// Update global scale.
-	TempDiff = MaxDelta - MinDelta;
-	InOutVertexDeltaScale = InOutVertexDeltaScale.ComponentMax(TempDiff.GetAbs());
-}
-
-bool UVertexDeltaTrainingModel::ComputeVertexDeltaStatistics(uint32 LODIndex)
-{
-	using namespace UE::VertexDeltaModel;
-
-	FVertexDeltaEditorModel* VertexDeltaEditorModel = GetVertexDeltaEditorModel();
-	UVertexDeltaModel* VertexDeltaModel = VertexDeltaEditorModel->GetVertexDeltaModel();
-	if (VertexDeltaEditorModel->IsDataNormalized())
-	{
-		return true;
-	}
-
-	FMLDeformerSampler* Sampler = VertexDeltaEditorModel->GetSampler();
-	Sampler->SetVertexDeltaSpace(EVertexDeltaSpace::PreSkinning);
-
-	// Show some dialog with progress.
-	const FText Title(LOCTEXT("PreprocessTrainingDataMessage", "Calculating data statistics"));
-	const uint32 AnimNumFrames = EditorModel->GetNumFramesForTraining();
-	FScopedSlowTask Task((float)AnimNumFrames, Title);
-	Task.MakeDialog(true);
-
-	// Get the base actor data.
-	FVector Mean = FVector::ZeroVector;
-	FVector Scale = FVector::OneVector;
-	float Count = 0.0f;
-	const int32 FrameCount = VertexDeltaEditorModel->GetNumFramesForTraining();
-	for (int32 FrameIndex = 0; FrameIndex < FrameCount; FrameIndex++)
-	{
-		// Calculate the deltas and update the mean and scale.
-		Sampler->Sample(FrameIndex);
-		UpdateVertexDeltaMeanAndScale(Sampler->GetVertexDeltas(), Mean, Scale, Count);
-
-		// Forward the progress bar and check if we want to cancel.
-		Task.EnterProgressFrame();
-		if (Task.ShouldCancel())
-		{
-			return false;
-		}
-	}
-
-	// Update the asset with calculated statistics.
-	VertexDeltaModel->VertexDeltaMean = Mean;
-	if (Count > 0.0f)
-	{
-		VertexDeltaModel->VertexDeltaScale = FVector::OneVector * Scale.GetMax();
-		VertexDeltaEditorModel->SetDataNormalized(true);
-	}
-
-	return true;
-}
-
-bool UVertexDeltaTrainingModel::ComputeDeltasStatistics()
-{
-	const UVertexDeltaModel* VertexDeltaModel = GetVertexDeltaModel();
-	const float DeltaCutoffLength = VertexDeltaModel->GetDeltaCutoffLength();
-	if (ComputeVertexDeltaStatistics(0))
-	{
-		VertexDeltaMean = VertexDeltaModel->VertexDeltaMean;
-		VertexDeltaScale = VertexDeltaModel->VertexDeltaScale;
-		return true;
-	}
-
-	return false;
-}
-
 #undef LOCTEXT_NAMESPACE
diff --git a/Engine/Plugins/Experimental/MLDeformer/VertexDeltaModel/Source/VertexDeltaModelEditor/Public/VertexDeltaEditorModel.h b/Engine/Plugins/Experimental/MLDeformer/VertexDeltaModel/Source/VertexDeltaModelEditor/Public/VertexDeltaEditorModel.h
--- a/Engine/Plugins/Experimental/MLDeformer/VertexDeltaModel/Source/VertexDeltaModelEditor/Public/VertexDeltaEditorModel.h
+++ b/Engine/Plugins/Experimental/MLDeformer/VertexDeltaModel/Source/VertexDeltaModelEditor/Public/VertexDeltaEditorModel.h
@@ -42,16 +42,12 @@ namespace UE::VertexDeltaModel
 		virtual FMLDeformerSampler* CreateSampler() const;
 		virtual void CreateTrainingGroundTruthActor(UWorld* World) override;
 		virtual void CreateTestGroundTruthActor(UWorld* World) override;
-		virtual void OnPreTraining() override;
-		virtual void OnTrainingAborted() override;
 		virtual void OnInputAssetsChanged() override;
 		virtual int32 GetNumFrames() const override;
 		virtual double GetTimeAtFrame(int32 FrameNumber) const override;
 		virtual int32 GetFrameAtTime(double TimeInSeconds) const override;
 		virtual void UpdateIsReadyForTrainingState() override;
 		virtual ETrainingResult Train() override;
-		virtual FString GetDefaultDeformerGraphAssetPath() const override;
-		virtual FString GetHeatMapDeformerGraphPath() const override;
 		FString GetTrainedNetworkOnnxFile() const override;
 		// ~END FMLDeformerEditorModel overrides.
 
diff --git a/Engine/Plugins/Experimental/MLDeformer/VertexDeltaModel/Source/VertexDeltaModelEditor/Public/VertexDeltaTrainingModel.h b/Engine/Plugins/Experimental/MLDeformer/VertexDeltaModel/Source/VertexDeltaModelEditor/Public/VertexDeltaTrainingModel.h
--- a/Engine/Plugins/Experimental/MLDeformer/VertexDeltaModel/Source/VertexDeltaModelEditor/Public/VertexDeltaTrainingModel.h
+++ b/Engine/Plugins/Experimental/MLDeformer/VertexDeltaModel/Source/VertexDeltaModelEditor/Public/VertexDeltaTrainingModel.h
@@ -24,23 +24,6 @@ public:
 	UFUNCTION(BlueprintImplementableEvent, Category = "Python")
 	int32 Train() const;
 
-	/** Compute delta statistics for the whole dataset. */
-	UFUNCTION(BlueprintCallable, Category = "Training Data")
-	bool ComputeDeltasStatistics();
-
 	UVertexDeltaModel* GetVertexDeltaModel() const;
 	UE::VertexDeltaModel::FVertexDeltaEditorModel* GetVertexDeltaEditorModel() const;
-
-protected:
-	bool ComputeVertexDeltaStatistics(uint32 LODIndex);
-	void UpdateVertexDeltaMeanAndScale(const TArray<float>& VertexDeltas, FVector& InOutMeanVertexDelta, FVector& InOutVertexDeltaScale, float& InOutCount);
-
-public:
-	// Mean delta computed over the entire dataset.
-	UPROPERTY(EditAnywhere, BlueprintReadWrite, Category = "Training Data")
-	FVector VertexDeltaMean = FVector::ZeroVector;
-
-	// Vertex delta scale computed over the entire dataset.
-	UPROPERTY(EditAnywhere, BlueprintReadWrite, Category = "Training Data")
-	FVector VertexDeltaScale = FVector::OneVector;
 };
