MLDeformer: Prevent crash if wrong MLDeformer data interface is selected.
#rnx

diff --git a/Engine/Plugins/Experimental/MLDeformer/LegacyVertexDeltaModel/Source/LegacyVertexDeltaModel/Private/LegacyVertexDeltaModelDataInterface.cpp b/Engine/Plugins/Experimental/MLDeformer/LegacyVertexDeltaModel/Source/LegacyVertexDeltaModel/Private/LegacyVertexDeltaModelDataInterface.cpp
--- a/Engine/Plugins/Experimental/MLDeformer/LegacyVertexDeltaModel/Source/LegacyVertexDeltaModel/Private/LegacyVertexDeltaModelDataInterface.cpp
+++ b/Engine/Plugins/Experimental/MLDeformer/LegacyVertexDeltaModel/Source/LegacyVertexDeltaModel/Private/LegacyVertexDeltaModelDataInterface.cpp
@@ -36,10 +36,11 @@ namespace UE::LegacyVertexDeltaModel
 	{
 		const UMLDeformerModel* Model = DeformerComponent->GetDeformerAsset()->GetModel();
 		const ULegacyVertexDeltaModel* VertexDeltaModel = Cast<ULegacyVertexDeltaModel>(Model);
-		check(VertexDeltaModel);
-
-		VertexDeltaScale = static_cast<FVector3f>(VertexDeltaModel->GetVertexDeltaScale());
-		VertexDeltaMean = static_cast<FVector3f>(VertexDeltaModel->GetVertexDeltaMean());
+		if (ensure(VertexDeltaModel))
+		{
+			VertexDeltaScale = static_cast<FVector3f>(VertexDeltaModel->GetVertexDeltaScale());
+			VertexDeltaMean = static_cast<FVector3f>(VertexDeltaModel->GetVertexDeltaMean());
+		}
 	}
 
 	void FLegacyVertexDeltaModelDataProviderProxy::GatherDispatchData(FDispatchSetup const& InDispatchSetup, FCollectedDispatchData& InOutDispatchData)
